Automation in Unreal Insights Frontend

diff --git a/Engine/Source/Developer/ScreenShotComparison/Private/Widgets/SScreenShotBrowser.cpp b/Engine/Source/Developer/ScreenShotComparison/Private/Widgets/SScreenShotBrowser.cpp
--- a/Engine/Source/Developer/ScreenShotComparison/Private/Widgets/SScreenShotBrowser.cpp
+++ b/Engine/Source/Developer/ScreenShotComparison/Private/Widgets/SScreenShotBrowser.cpp
@@ -338,11 +338,15 @@ void SScreenShotBrowser::Construct( const FArguments& InArgs,  IScreenShotManage
 
 SScreenShotBrowser::~SScreenShotBrowser()
 {
-	FDirectoryWatcherModule& DirectoryWatcherModule = FModuleManager::LoadModuleChecked<FDirectoryWatcherModule>(TEXT("DirectoryWatcher"));
-
-	if ( DirectoryWatchingHandle.IsValid() )
+	// Do not load the module if it was already unloaded.
+	if (FModuleManager::Get().IsModuleLoaded(TEXT("DirectoryWatcher")))
 	{
-		DirectoryWatcherModule.Get()->UnregisterDirectoryChangedCallback_Handle(ComparisonRoot, DirectoryWatchingHandle);
+		FDirectoryWatcherModule& DirectoryWatcherModule = FModuleManager::LoadModuleChecked<FDirectoryWatcherModule>(TEXT("DirectoryWatcher"));
+
+		if (DirectoryWatchingHandle.IsValid())
+		{
+			DirectoryWatcherModule.Get()->UnregisterDirectoryChangedCallback_Handle(ComparisonRoot, DirectoryWatchingHandle);
+		}
 	}
 }
 
diff --git a/Engine/Source/Developer/TraceInsightsFrontend/Private/InsightsFrontend/Common/InsightsAutomationController.cpp b/Engine/Source/Developer/TraceInsightsFrontend/Private/InsightsFrontend/Common/InsightsAutomationController.cpp
new file mode 100644
--- /dev/null
+++ b/Engine/Source/Developer/TraceInsightsFrontend/Private/InsightsFrontend/Common/InsightsAutomationController.cpp
@@ -0,0 +1,119 @@
+// Copyright Epic Games, Inc. All Rights Reserved.
+
+#include "InsightsAutomationController.h"
+
+#if INSIGHTS_ENABLE_AUTOMATION
+
+#include "IAutomationControllerModule.h"
+#include "IAutomationWorkerModule.h"
+#include "ISessionManager.h"
+#include "ISessionServicesModule.h"
+#include "Misc/App.h"
+#include "Misc/AutomationTest.h"
+#include "Misc/CoreMisc.h"
+#include "Widgets/Docking/SDockTab.h"
+
+#endif
+
+// TraceInsightsCore
+#include "InsightsCore/Common/InsightsCoreStyle.h"
+
+////////////////////////////////////////////////////////////////////////////////////////////////////
+
+DEFINE_LOG_CATEGORY(InsightsAutomationController);
+
+#define LOCTEXT_NAMESPACE "UE::Insights::FInsightsAutomationController"
+
+////////////////////////////////////////////////////////////////////////////////////////////////////
+
+namespace UE::Insights
+{
+
+const TCHAR* FInsightsAutomationController::AutoQuitMsgOnComplete = TEXT("Application is closing because it was started with the AutoQuit parameter and session analysis is complete and all scheduled tests have completed.");
+
+////////////////////////////////////////////////////////////////////////////////////////////////////
+
+void FInsightsAutomationController::RunTests(const FString& InCmd)
+{
+#if INSIGHTS_ENABLE_AUTOMATION
+	FString ActualCmd = InCmd.Replace(TEXT("\""), TEXT(""));
+	if (!ActualCmd.StartsWith(TEXT("Automation RunTests")))
+	{
+		UE_LOG(InsightsAutomationController, Warning, TEXT("[FInsightsAutomationController] Command %s does not start with Automation RunTests. Command will be ignored."), *InCmd);
+		return;
+	}
+
+	IAutomationControllerModule& AutomationControllerModule = FModuleManager::LoadModuleChecked<IAutomationControllerModule>(TEXT("AutomationController"));
+	IAutomationControllerManagerRef AutomationControllerManager = AutomationControllerModule.GetAutomationController();
+
+	AutomationControllerManager->OnTestsComplete().AddLambda([this]()
+		{
+			RunningTestsState = ETestsState::Finished;
+		});
+
+	RunningTestsState = ETestsState::Running;
+	StaticExec(NULL, *ActualCmd);
+#endif
+}
+
+////////////////////////////////////////////////////////////////////////////////////////////////////
+
+void FInsightsAutomationController::Initialize()
+{
+#if INSIGHTS_ENABLE_AUTOMATION
+	FApp::SetSessionName(TEXT("UnrealInsights"));
+	ISessionServicesModule& SessionServicesModule = FModuleManager::LoadModuleChecked<ISessionServicesModule>("SessionServices");
+	auto SessionService = SessionServicesModule.GetSessionService();
+	SessionService->Start();
+
+	// Create Session Manager.
+	SessionServicesModule.GetSessionManager();
+
+	IAutomationControllerModule& AutomationControllerModule = FModuleManager::LoadModuleChecked<IAutomationControllerModule>(TEXT("AutomationController"));
+	AutomationControllerModule.Init();
+
+	// Initialize the target platform manager as it is needed by Automation Window.
+	GetTargetPlatformManager();
+	FModuleManager::Get().LoadModule("AutomationWindow");
+	FModuleManager::Get().LoadModule("AutomationWorker");
+
+	// Register tick functions.
+	OnTick = FTickerDelegate::CreateSP(this, &FInsightsAutomationController::Tick);
+	OnTickHandle = FTSTicker::GetCoreTicker().AddTicker(OnTick, 1.0f);
+#endif
+}
+
+////////////////////////////////////////////////////////////////////////////////////////////////////
+
+FInsightsAutomationController::~FInsightsAutomationController()
+{
+#if INSIGHTS_ENABLE_AUTOMATION
+	FTSTicker::GetCoreTicker().RemoveTicker(OnTickHandle);
+#endif
+}
+
+////////////////////////////////////////////////////////////////////////////////////////////////////
+
+bool FInsightsAutomationController::Tick(float DeltaTime)
+{
+#if INSIGHTS_ENABLE_AUTOMATION
+	if (bAutoQuit && RunningTestsState == ETestsState::Finished)
+	{
+		RequestEngineExit(AutoQuitMsgOnComplete);
+	}
+
+	IAutomationWorkerModule& AutomationWorkerModule = FModuleManager::LoadModuleChecked<IAutomationWorkerModule>("AutomationWorker");
+	IAutomationControllerModule& AutomationControllerModule = FModuleManager::LoadModuleChecked<IAutomationControllerModule>(TEXT("AutomationController"));
+
+	AutomationControllerModule.Tick();
+	AutomationWorkerModule.Tick();
+
+#endif
+	return true;
+}
+
+////////////////////////////////////////////////////////////////////////////////////////////////////
+
+} // namespace UE::Insights
+
+#undef LOCTEXT_NAMESPACE
\ No newline at end of file
diff --git a/Engine/Source/Developer/TraceInsightsFrontend/Private/InsightsFrontend/Common/InsightsAutomationController.h b/Engine/Source/Developer/TraceInsightsFrontend/Private/InsightsFrontend/Common/InsightsAutomationController.h
new file mode 100644
--- /dev/null
+++ b/Engine/Source/Developer/TraceInsightsFrontend/Private/InsightsFrontend/Common/InsightsAutomationController.h
@@ -0,0 +1,52 @@
+// Copyright Epic Games, Inc. All Rights Reserved.
+
+#pragma once
+
+#include "InsightsFrontend/ITraceInsightsFrontendModule.h"
+
+#include "CoreMinimal.h"
+#include "Logging/LogMacros.h"
+#include "Containers/Ticker.h"
+
+DECLARE_LOG_CATEGORY_EXTERN(InsightsAutomationController, Log, All);
+
+namespace UE::Insights
+{
+
+class TRACEINSIGHTSFRONTEND_API FInsightsAutomationController : public TSharedFromThis<FInsightsAutomationController>
+{
+	enum class ETestsState
+	{
+		NotStarted,
+		Running,
+		Finished,
+	};
+
+public:
+	virtual ~FInsightsAutomationController();
+
+	virtual void Initialize();
+
+	bool Tick(float DeltaTime);
+
+	void SetAutoQuit(bool InAutoQuit) { bAutoQuit = InAutoQuit; }
+	bool GetAutoQuit() const { return bAutoQuit; }
+
+	void RunTests(const FString& InCmd);
+
+private:
+	/** The delegate to be invoked when this manager ticks. */
+	FTickerDelegate OnTick;
+
+	/** Handle to the registered OnTick. */
+	FTSTicker::FDelegateHandle OnTickHandle;
+
+	FString CommandToExecute;
+
+	bool bAutoQuit = false;
+	ETestsState RunningTestsState = ETestsState::NotStarted;
+
+	static const TCHAR* AutoQuitMsgOnComplete;
+};
+
+} // namespace UE::Insights
\ No newline at end of file
diff --git a/Engine/Source/Developer/TraceInsightsFrontend/Private/InsightsFrontend/Tests/FunctionalTests/MemoryInsightsTests.cpp b/Engine/Source/Developer/TraceInsightsFrontend/Private/InsightsFrontend/Tests/FunctionalTests/MemoryInsightsTests.cpp
--- a/Engine/Source/Developer/TraceInsightsFrontend/Private/InsightsFrontend/Tests/FunctionalTests/MemoryInsightsTests.cpp
+++ b/Engine/Source/Developer/TraceInsightsFrontend/Private/InsightsFrontend/Tests/FunctionalTests/MemoryInsightsTests.cpp
@@ -146,6 +146,7 @@ void FAutomationDriverUnrealInsightsHubMemoryInsightsTest::Define()
 		Driver.Reset();
 		IAutomationDriverModule::Get().Disable();
 		AutomationWindow->Restore();
+		AutomationWindow.Reset();
 	});
 }
 
diff --git a/Engine/Source/Developer/TraceInsightsFrontend/Private/InsightsFrontend/Tests/FunctionalTests/TraceStoreWindowTests.cpp b/Engine/Source/Developer/TraceInsightsFrontend/Private/InsightsFrontend/Tests/FunctionalTests/TraceStoreWindowTests.cpp
--- a/Engine/Source/Developer/TraceInsightsFrontend/Private/InsightsFrontend/Tests/FunctionalTests/TraceStoreWindowTests.cpp
+++ b/Engine/Source/Developer/TraceInsightsFrontend/Private/InsightsFrontend/Tests/FunctionalTests/TraceStoreWindowTests.cpp
@@ -164,6 +164,7 @@ void FAutomationDriverUnrealInsightsSessionBrowserTest::Define()
 		Driver.Reset();
 		IAutomationDriverModule::Get().Disable();
 		AutomationWindow->Restore();
+		AutomationWindow.Reset();
 	});
 }
 
diff --git a/Engine/Source/Developer/TraceInsightsFrontend/Private/InsightsFrontend/TraceInsightsFrontendModule.cpp b/Engine/Source/Developer/TraceInsightsFrontend/Private/InsightsFrontend/TraceInsightsFrontendModule.cpp
--- a/Engine/Source/Developer/TraceInsightsFrontend/Private/InsightsFrontend/TraceInsightsFrontendModule.cpp
+++ b/Engine/Source/Developer/TraceInsightsFrontend/Private/InsightsFrontend/TraceInsightsFrontendModule.cpp
@@ -12,7 +12,7 @@
 #include "WorkspaceMenuStructure.h"
 #include "WorkspaceMenuStructureModule.h"
 
-// TarceAnalysis
+// TraceAnalysis
 #include "Trace/StoreConnection.h"
 
 // TraceInsightsCore
@@ -21,7 +21,8 @@
 
 // TraceInsightsFrontend
 #include "InsightsFrontend/Common/InsightsFrontendStyle.h"
-#include "InsightsFrontend/Tests/TestRunner.h"
+#include "InsightsFrontend/Common/InsightsAutomationController.h"
+#include "InsightsFrontend/Common/Log.h"
 #include "InsightsFrontend/Widgets/SConnectionWindow.h"
 #include "InsightsFrontend/Widgets/STraceStoreWindow.h"
 
@@ -34,7 +35,6 @@ namespace UE::Insights
 
 const FName FInsightsFrontendTabs::TraceStoreTabId("TraceStore");
 const FName FInsightsFrontendTabs::ConnectionTabId("Connection");
-const FName FInsightsFrontendTabs::AutomationWindowTabId("Automation");
 
 ////////////////////////////////////////////////////////////////////////////////////////////////////
 
@@ -76,6 +76,8 @@ void FTraceInsightsFrontendModule::ShutdownModule()
 	FInsightsFrontendStyle::Shutdown();
 
 	TraceStoreConnection.Reset();
+
+	InsightsAutomationController.Reset();
 }
 
 ////////////////////////////////////////////////////////////////////////////////////////////////////
@@ -127,6 +129,11 @@ void FTraceInsightsFrontendModule::RegisterTabSpawners()
 			WorkspaceMenu::GetMenuStructure().GetDeveloperToolsProfilingCategory();
 		TabSpawnerEntry.SetGroup(Group);
 	}
+
+	InsightsAutomationController = MakeShared<FInsightsAutomationController>();
+
+	InsightsAutomationController->Initialize();
+	InsightsAutomationController->SetAutoQuit(CreateWindowParams.bAutoQuit);
 }
 
 ////////////////////////////////////////////////////////////////////////////////////////////////////
@@ -139,8 +146,10 @@ void FTraceInsightsFrontendModule::UnregisterTabSpawners()
 
 ////////////////////////////////////////////////////////////////////////////////////////////////////
 
-void FTraceInsightsFrontendModule::CreateFrontendWindow(const FCreateFrontendWindowParams& Params)
+void FTraceInsightsFrontendModule::CreateFrontendWindow(const FCreateFrontendWindowParams& InParams)
 {
+	CreateWindowParams = InParams;
+
 	RegisterTabSpawners();
 
 	//////////////////////////////////////////////////
@@ -203,7 +212,7 @@ void FTraceInsightsFrontendModule::CreateFrontendWindow(const FCreateFrontendWin
 		->SetWindow(FVector2D(10.0f * DPIScaleFactor, 10.0f * DPIScaleFactor), false)
 		->Split
 		(
-			FTabManager::NewStack()->AddTab(FTabId("WidgetReflector"), Params.bAllowDebugTools ? ETabState::OpenedTab : ETabState::ClosedTab)
+			FTabManager::NewStack()->AddTab(FTabId("WidgetReflector"), InParams.bAllowDebugTools ? ETabState::OpenedTab : ETabState::ClosedTab)
 		)
 	);
 
@@ -229,10 +238,10 @@ void FTraceInsightsFrontendModule::CreateFrontendWindow(const FCreateFrontendWin
 	TSharedPtr<STraceStoreWindow> TraceStoreWnd = GetTraceStoreWindow();
 	if (TraceStoreWnd.IsValid())
 	{
-		TraceStoreWnd->SetEnableAutomaticTesting(Params.bInitializeTesting);
-		TraceStoreWnd->SetEnableDebugTools(Params.bAllowDebugTools);
-		TraceStoreWnd->SetStartProcessWithStompMalloc(Params.bStartProcessWithStompMalloc);
-		TraceStoreWnd->SetDisableFramerateThrottle(Params.bDisableFramerateThrottle);
+		TraceStoreWnd->SetEnableAutomaticTesting(InParams.bInitializeTesting);
+		TraceStoreWnd->SetEnableDebugTools(InParams.bAllowDebugTools);
+		TraceStoreWnd->SetStartProcessWithStompMalloc(InParams.bStartProcessWithStompMalloc);
+		TraceStoreWnd->SetDisableFramerateThrottle(InParams.bDisableFramerateThrottle);
 	}
 }
 
@@ -324,21 +333,21 @@ void FTraceInsightsFrontendModule::OnConnectionTabClosed(TSharedRef<SDockTab> Ta
 
 void FTraceInsightsFrontendModule::RunAutomationTests(const FString& InCmd)
 {
-#if !UE_BUILD_SHIPPING && !WITH_EDITOR
+#if INSIGHTS_ENABLE_AUTOMATION
+
 	FString ActualCmd = InCmd;
 	ActualCmd.TrimCharInline(TEXT('\"'), nullptr);
 	ActualCmd.TrimCharInline(TEXT('\''), nullptr);
 
 	if (ActualCmd.StartsWith(TEXT("Automation RunTests")))
 	{
-		FTestRunner TestRunner;
-		TestRunner.Run(ActualCmd);
+		InsightsAutomationController->RunTests(InCmd);
 	}
-#endif
+#else
+	UE_LOG(LogInsightsFrontend, Error, TEXT("Automated test could not execute because INSIGHTS_ENABLE_AUTOMATION is disabled."));
+#endif // INSIGHTS_ENABLE_AUTOMATION
 }
 
-////////////////////////////////////////////////////////////////////////////////////////////////////
-
 } // namespace UE::Insights
 
-#undef LOCTEXT_NAMESPACE
+#undef LOCTEXT_NAMESPACE
\ No newline at end of file
diff --git a/Engine/Source/Developer/TraceInsightsFrontend/Private/InsightsFrontend/TraceInsightsFrontendModule.h b/Engine/Source/Developer/TraceInsightsFrontend/Private/InsightsFrontend/TraceInsightsFrontendModule.h
--- a/Engine/Source/Developer/TraceInsightsFrontend/Private/InsightsFrontend/TraceInsightsFrontendModule.h
+++ b/Engine/Source/Developer/TraceInsightsFrontend/Private/InsightsFrontend/TraceInsightsFrontendModule.h
@@ -2,6 +2,7 @@
 
 #pragma once
 
+#include "Containers/Ticker.h"
 #include "Framework/Docking/TabManager.h"
 
 // TraceInsightsFrontend
@@ -13,6 +14,7 @@ namespace UE::Trace { class FStoreConnection; }
 namespace UE::Insights
 {
 
+class FInsightsAutomationController;
 class STraceStoreWindow;
 class SConnectionWindow;
 
@@ -70,6 +72,8 @@ private:
 	/** An instance of the main settings. */
 	FInsightsFrontendSettings Settings;
 
+	FCreateFrontendWindowParams CreateWindowParams;
+
 	TSharedPtr<FTabManager::FLayout> PersistentLayout;
 	static FString LayoutIni;
 
@@ -82,6 +86,14 @@ private:
 	TWeakPtr<SConnectionWindow> ConnectionWindow;
 
 	bool bIsMainTabSet = false;
+
+	TSharedPtr<UE::Insights::FInsightsAutomationController> InsightsAutomationController;
+
+	/** The delegate to be invoked when this ticks. */
+	FTickerDelegate OnTick;
+
+	/** Handle to the registered OnTick. */
+	FTSTicker::FDelegateHandle OnTickHandle;
 };
 
 } // UE::Insights
diff --git a/Engine/Source/Developer/TraceInsightsFrontend/Private/InsightsFrontend/Widgets/STraceStoreWindow.cpp b/Engine/Source/Developer/TraceInsightsFrontend/Private/InsightsFrontend/Widgets/STraceStoreWindow.cpp
--- a/Engine/Source/Developer/TraceInsightsFrontend/Private/InsightsFrontend/Widgets/STraceStoreWindow.cpp
+++ b/Engine/Source/Developer/TraceInsightsFrontend/Private/InsightsFrontend/Widgets/STraceStoreWindow.cpp
@@ -2237,20 +2237,6 @@ TSharedRef<SWidget> STraceStoreWindow::MakeTraceListMenu()
 
 	MenuBuilder.BeginSection("DebugOptions", LOCTEXT("TraceListMenu_Section_DebugOptions", "Debug Options"));
 
-	if (FGlobalTabmanager::Get()->HasTabSpawner(FInsightsFrontendTabs::AutomationWindowTabId))
-	{
-		MenuBuilder.AddMenuEntry(
-			LOCTEXT("Automation", "Automation"),
-			TAttribute<FText>(), // no tooltip
-			FSlateIcon(FInsightsCoreStyle::GetStyleSetName(), "Icons.TestAutomation"),
-			FUIAction(FExecuteAction::CreateLambda([this]()
-				{
-					FGlobalTabmanager::Get()->TryInvokeTab(FInsightsFrontendTabs::AutomationWindowTabId);
-				})),
-			NAME_None,
-			EUserInterfaceActionType::Button);
-	}
-
 	// Enable Automation Tests Option.
 	{
 		FUIAction ToogleAutomationTestsAction;
diff --git a/Engine/Source/Developer/TraceInsightsFrontend/Public/InsightsFrontend/ITraceInsightsFrontendModule.h b/Engine/Source/Developer/TraceInsightsFrontend/Public/InsightsFrontend/ITraceInsightsFrontendModule.h
--- a/Engine/Source/Developer/TraceInsightsFrontend/Public/InsightsFrontend/ITraceInsightsFrontendModule.h
+++ b/Engine/Source/Developer/TraceInsightsFrontend/Public/InsightsFrontend/ITraceInsightsFrontendModule.h
@@ -7,6 +7,8 @@
 #include "Modules/ModuleInterface.h"
 #include "Templates/SharedPointer.h"
 
+#define INSIGHTS_ENABLE_AUTOMATION !WITH_EDITOR
+
 namespace UE::Insights
 {
 
@@ -17,13 +19,13 @@ struct FCreateFrontendWindowParams
 	bool bInitializeTesting = false;
 	bool bStartProcessWithStompMalloc = false;
 	bool bDisableFramerateThrottle = false;
+	bool bAutoQuit = false;
 };
 
 struct FInsightsFrontendTabs
 {
 	static const FName TraceStoreTabId;
 	static const FName ConnectionTabId;
-	static const FName AutomationWindowTabId;
 };
 
 class STraceStoreWindow;
diff --git a/Engine/Source/Developer/TraceInsightsFrontend/TraceInsightsFrontend.Build.cs b/Engine/Source/Developer/TraceInsightsFrontend/TraceInsightsFrontend.Build.cs
--- a/Engine/Source/Developer/TraceInsightsFrontend/TraceInsightsFrontend.Build.cs
+++ b/Engine/Source/Developer/TraceInsightsFrontend/TraceInsightsFrontend.Build.cs
@@ -36,8 +36,8 @@ public class TraceInsightsFrontend : ModuleRules
 			);
 		}
 
-		// Modules required for running automation in stand alone Insights
-		if (Target.Configuration != UnrealTargetConfiguration.Shipping && !Target.bBuildEditor)
+		// Modules required for running automation
+		if (!Target.bBuildEditor)
 		{
 			PrivateDependencyModuleNames.AddRange(
 				new string[] {
diff --git a/Engine/Source/Developer/TraceTools/Private/Widgets/SFilterPresetList.cpp b/Engine/Source/Developer/TraceTools/Private/Widgets/SFilterPresetList.cpp
--- a/Engine/Source/Developer/TraceTools/Private/Widgets/SFilterPresetList.cpp
+++ b/Engine/Source/Developer/TraceTools/Private/Widgets/SFilterPresetList.cpp
@@ -143,6 +143,11 @@ FReply SFilterPresetList::OnMouseButtonUp( const FGeometry& MyGeometry, const FP
 
 void SFilterPresetList::SaveSettings(const FString& IniFilename) const
 {
+	if (GConfig == nullptr)
+	{
+		return;
+	}
+
 	FStringView IniSectionName = TEXTVIEW("Trace.FilterPresetList");
 	FStringView IniActivePresetsKey = TEXTVIEW("ActivePresets");
 
diff --git a/Engine/Source/Developer/TraceTools/Private/Widgets/STraceControl.cpp b/Engine/Source/Developer/TraceTools/Private/Widgets/STraceControl.cpp
--- a/Engine/Source/Developer/TraceTools/Private/Widgets/STraceControl.cpp
+++ b/Engine/Source/Developer/TraceTools/Private/Widgets/STraceControl.cpp
@@ -92,6 +92,24 @@ void STraceControl::SetInstanceId(const FGuid& Id)
 	{
 		TraceController->SendDiscoveryRequest(FGuid(), Id);
 	}
+
+#if !WITH_EDITOR
+	if (Id == FApp::GetInstanceId())
+	{
+		TraceDataFilterWidget->SetWarningBannerText(LOCTEXT("SessionCannotBeControlledWarning", "Unreal Session Frontend cannot be controlled by this tool. Please select another active instance from the Session Browser."));
+	}
+	else
+	{
+		if (bAutoDetectSelectedSession)
+		{
+			TraceDataFilterWidget->SetWarningBannerText(LOCTEXT("NoSessionSelectedWarning", "Please select an active instance from the Session Browser."));
+		}
+		else
+		{
+			TraceDataFilterWidget->SetWarningBannerText(FText::GetEmpty());
+		}
+	}
+#endif
 }
 
 void STraceControl::OnInstanceSelectionChanged(const TSharedPtr<ISessionInstanceInfo>& InstanceInfo, bool bSelected)
diff --git a/Engine/Source/Programs/UnrealInsights/Private/UserInterfaceCommand.cpp b/Engine/Source/Programs/UnrealInsights/Private/UserInterfaceCommand.cpp
--- a/Engine/Source/Programs/UnrealInsights/Private/UserInterfaceCommand.cpp
+++ b/Engine/Source/Programs/UnrealInsights/Private/UserInterfaceCommand.cpp
@@ -291,6 +291,7 @@ void FUserInterfaceCommand::InitializeSlateApplication(bool bFrontendMode, const
 		Params.bInitializeTesting = FParse::Param(FCommandLine::Get(), TEXT("InsightsTest"));
 		Params.bStartProcessWithStompMalloc = FParse::Param(FCommandLine::Get(), TEXT("stompmalloc"));
 		Params.bDisableFramerateThrottle = FParse::Param(FCommandLine::Get(), TEXT("DisableFramerateThrottle"));
+		Params.bAutoQuit = FParse::Param(FCommandLine::Get(), TEXT("AutoQuit"));
 		TraceInsightsFrontendModule.CreateFrontendWindow(Params);
 
 		if (bRunAutomationTests)
