[Insights] Improved the UI/UX when Unreal Insights Frontend cannot connect to Trace Server:
    * the error message is now clearly displayed when not connected to a trace server and trying to reconnect
    * the UI elements related to an active connection are all hidden now when not connected

diff --git a/Engine/Source/Developer/TraceInsightsFrontend/Private/InsightsFrontend/Widgets/STraceStoreWindow.cpp b/Engine/Source/Developer/TraceInsightsFrontend/Private/InsightsFrontend/Widgets/STraceStoreWindow.cpp
--- a/Engine/Source/Developer/TraceInsightsFrontend/Private/InsightsFrontend/Widgets/STraceStoreWindow.cpp
+++ b/Engine/Source/Developer/TraceInsightsFrontend/Private/InsightsFrontend/Widgets/STraceStoreWindow.cpp
@@ -400,6 +400,7 @@ TSharedRef<SWidget> STraceStoreWindow::ConstructFiltersToolbar()
 	RightSideToolbarBuilder.EndSection();
 
 	return SNew(SHorizontalBox)
+		.Visibility(this, &STraceStoreWindow::VisibleIfConnected)
 
 		+ SHorizontalBox::Slot()
 		.HAlign(HAlign_Fill)
@@ -427,6 +428,7 @@ BEGIN_SLATE_FUNCTION_BUILD_OPTIMIZATION
 TSharedRef<SWidget> STraceStoreWindow::ConstructSessionsPanel()
 {
 	TSharedRef<SWidget> Widget = SAssignNew(TraceListView, SListView<TSharedPtr<FTraceViewModel>>)
+		.Visibility(this, &STraceStoreWindow::HiddenIfNotConnected)
 		.IsFocusable(true)
 		.SelectionMode(ESelectionMode::Multi)
 		.OnSelectionChanged(this, &STraceStoreWindow::TraceList_OnSelectionChanged)
@@ -661,6 +663,25 @@ TSharedRef<SWidget> STraceStoreWindow::ConstructTraceStoreDirectoryPanel()
 
 		SNew(SVerticalBox)
 
+		+ SVerticalBox::Slot()
+		.AutoHeight()
+		.Padding(0.0f, 8.0f)
+		.HAlign(HAlign_Fill)
+		[
+			SNew(SHorizontalBox)
+			.Visibility(this, &STraceStoreWindow::VisibleIfNotConnected)
+
+			+ SHorizontalBox::Slot()
+			.AutoWidth()
+			.Padding(0.0f, 0.0f, 0.0f, 0.0f)
+			.VAlign(VAlign_Center)
+			[
+				SNew(STextBlock)
+				.Text_Raw(this, &STraceStoreWindow::GetConnectionStatusTooltip)
+				.ColorAndOpacity(EStyleColor::Error)
+			]
+		]
+
 		+ SVerticalBox::Slot()
 		.AutoHeight()
 		.Padding(0.0f, 0.0f)
@@ -720,6 +741,7 @@ TSharedRef<SWidget> STraceStoreWindow::ConstructTraceStoreDirectoryPanel()
 			.VAlign(VAlign_Center)
 			[
 				SNew(STextBlock)
+				.Visibility(this, &STraceStoreWindow::VisibleIfConnected)
 				.Text(LOCTEXT("TraceStoreDirText", "Directory:"))
 			]
 
@@ -729,6 +751,7 @@ TSharedRef<SWidget> STraceStoreWindow::ConstructTraceStoreDirectoryPanel()
 			.VAlign(VAlign_Center)
 			[
 				SAssignNew(StoreDirTextBox, SEditableTextBox)
+				.Visibility(this, &STraceStoreWindow::VisibleIfConnected)
 				.IsReadOnly(true)
 				.BackgroundColor(FSlateColor(EStyleColor::Background))
 				.Text(this, &STraceStoreWindow::GetTraceStoreDirectory)
@@ -740,6 +763,7 @@ TSharedRef<SWidget> STraceStoreWindow::ConstructTraceStoreDirectoryPanel()
 			.VAlign(VAlign_Center)
 			[
 				SNew(SButton)
+				.Visibility(this, &STraceStoreWindow::VisibleIfConnected)
 				.ButtonStyle(&FAppStyle::Get().GetWidgetStyle<FButtonStyle>("SimpleButton"))
 				.ToolTipText(LOCTEXT("ExploreTraceStoreDirButtonToolTip", "Explores the Trace Store Directory."))
 				.OnClicked(this, &STraceStoreWindow::ExploreTraceStoreDirectory_OnClicked)
@@ -765,8 +789,23 @@ TSharedRef<SWidget> STraceStoreWindow::ConstructTraceStoreDirectoryPanel()
 
 			+ SScrollBox::Slot()
 			[
+				SNew(SVerticalBox)
+				.Visibility(this, &STraceStoreWindow::VisibleIfNotConnected)
+
+				+ SVerticalBox::Slot()
+				.AutoHeight()
+				.Padding(0.0f, 4.0f)
+				[
+					SNew(STextBlock)
+					.Text(LOCTEXT("NotConnected", "Not connected to a Trace Server!"))
+					.ColorAndOpacity(EStyleColor::Warning)
+				]
+			]
 
+			+ SScrollBox::Slot()
+			[
 				SNew(SVerticalBox)
+				.Visibility(this, &STraceStoreWindow::VisibleIfConnected)
 
 				+ SVerticalBox::Slot()
 				.AutoHeight()
@@ -990,6 +1029,36 @@ const FSlateBrush* STraceStoreWindow::GetConnectionStatusIcon() const
 
 ////////////////////////////////////////////////////////////////////////////////////////////////////
 
+EVisibility STraceStoreWindow::VisibleIfNotConnected() const
+{
+	const FStoreBrowser::EConnectionStatus Status = StoreBrowser->GetConnectionStatus();
+	return (Status == FStoreBrowser::EConnectionStatus::Connected) ?
+		EVisibility::Collapsed :
+		EVisibility::Visible;
+}
+
+////////////////////////////////////////////////////////////////////////////////////////////////////
+
+EVisibility STraceStoreWindow::VisibleIfConnected() const
+{
+	const FStoreBrowser::EConnectionStatus Status = StoreBrowser->GetConnectionStatus();
+	return (Status == FStoreBrowser::EConnectionStatus::Connected) ?
+		EVisibility::Visible :
+		EVisibility::Collapsed;
+}
+
+////////////////////////////////////////////////////////////////////////////////////////////////////
+
+EVisibility STraceStoreWindow::HiddenIfNotConnected() const
+{
+	const FStoreBrowser::EConnectionStatus Status = StoreBrowser->GetConnectionStatus();
+	return (Status == FStoreBrowser::EConnectionStatus::Connected) ?
+		EVisibility::Visible :
+		EVisibility::Hidden;
+}
+
+////////////////////////////////////////////////////////////////////////////////////////////////////
+
 TSharedPtr<SWidget> STraceStoreWindow::TraceList_GetMenuContent()
 {
 	const bool bShouldCloseWindowAfterMenuSelection = true;
diff --git a/Engine/Source/Developer/TraceInsightsFrontend/Private/InsightsFrontend/Widgets/STraceStoreWindow.h b/Engine/Source/Developer/TraceInsightsFrontend/Private/InsightsFrontend/Widgets/STraceStoreWindow.h
--- a/Engine/Source/Developer/TraceInsightsFrontend/Private/InsightsFrontend/Widgets/STraceStoreWindow.h
+++ b/Engine/Source/Developer/TraceInsightsFrontend/Private/InsightsFrontend/Widgets/STraceStoreWindow.h
@@ -122,6 +122,9 @@ private:
 	TSharedRef<SWidget> ConstructLoadPanel();
 	FText GetConnectionStatusTooltip() const;
 	const FSlateBrush* GetConnectionStatusIcon() const;
+	EVisibility VisibleIfConnected() const;
+	EVisibility VisibleIfNotConnected() const;
+	EVisibility HiddenIfNotConnected() const;
 	TSharedRef<SWidget> ConstructTraceStoreDirectoryPanel();
 	TSharedRef<SWidget> ConstructAutoStartPanel();
 
