[Insights] Added log messages for user actions of adding or removing monitored directories for local trace store.

diff --git a/Engine/Source/Developer/TraceInsightsFrontend/Private/InsightsFrontend/Widgets/STraceDirectoryItem.cpp b/Engine/Source/Developer/TraceInsightsFrontend/Private/InsightsFrontend/Widgets/STraceDirectoryItem.cpp
--- a/Engine/Source/Developer/TraceInsightsFrontend/Private/InsightsFrontend/Widgets/STraceDirectoryItem.cpp
+++ b/Engine/Source/Developer/TraceInsightsFrontend/Private/InsightsFrontend/Widgets/STraceDirectoryItem.cpp
@@ -22,6 +22,7 @@
 
 // TraceInsightsFrontend
 #include "InsightsFrontend/Common/InsightsFrontendStyle.h"
+#include "InsightsFrontend/Common/Log.h"
 #include "InsightsFrontend/Widgets/STraceStoreWindow.h"
 
 #define LOCTEXT_NAMESPACE "UE::Insights::STraceDirectoryItem"
@@ -141,7 +142,7 @@ TSharedRef<SWidget> STraceDirectoryItem::ConstructOperations()
 			[
 				SNew(SButton)
 				.ButtonStyle(&FAppStyle::Get().GetWidgetStyle<FButtonStyle>("SimpleButton"))
-				.ToolTipText(bIsWatchDir ? 
+				.ToolTipText(bIsWatchDir ?
 								LOCTEXT("ExploreWatchDirButtonToolTip", "Explores the monitored directory.") :
 								LOCTEXT("ExploreTraceStoreDirButtonToolTip", "Explores the Trace Store Directory."))
 				.OnClicked_Raw(this, &STraceDirectoryItem::OnExplore)
@@ -207,11 +208,14 @@ FReply STraceDirectoryItem::OnDelete()
 
 		if (Response == EDialogResponse::OK)
 		{
+			UE_LOG(LogInsightsFrontend, Log, TEXT("[TraceStore] Removing monitored directory: \"%s\"..."), *Model->Path);
+
 			UE::Trace::FStoreClient* StoreClient = Window->GetTraceStoreConnection().GetStoreClient();
-			if (StoreClient)
+			if (!StoreClient ||
+				!StoreClient->SetStoreDirectories(nullptr, {}, { *Model->Path }))
 			{
-				const TCHAR* Path = *Model->Path;
-				StoreClient->SetStoreDirectories(nullptr, {}, { Path });
+				FName LogListingName(TEXT("UnrealInsights"));
+				FMessageLog(LogListingName).Error(LOCTEXT("StoreCommunicationFail", "Failed to change settings on the store service."));
 			}
 		}
 	}
@@ -269,8 +273,11 @@ FReply STraceDirectoryItem::OnModifyStore()
 				RemoveWatchDirs.Emplace(SelectedDirectory);
 			}
 
+			UE_LOG(LogInsightsFrontend, Log, TEXT("[TraceStore] Changing store directory: \"%s\"..."), *SelectedDirectory);
+
 			UE::Trace::FStoreClient* StoreClient = Window->GetTraceStoreConnection().GetStoreClient();
-			if (!StoreClient || !StoreClient->SetStoreDirectories(*SelectedDirectory, AddWatchDirs, RemoveWatchDirs))
+			if (!StoreClient ||
+				!StoreClient->SetStoreDirectories(*SelectedDirectory, AddWatchDirs, RemoveWatchDirs))
 			{
 				FName LogListingName(TEXT("UnrealInsights"));
 				FMessageLog(LogListingName).Error(LOCTEXT("StoreCommunicationFail", "Failed to change settings on the store service."));
diff --git a/Engine/Source/Developer/TraceInsightsFrontend/Private/InsightsFrontend/Widgets/STraceStoreWindow.cpp b/Engine/Source/Developer/TraceInsightsFrontend/Private/InsightsFrontend/Widgets/STraceStoreWindow.cpp
--- a/Engine/Source/Developer/TraceInsightsFrontend/Private/InsightsFrontend/Widgets/STraceStoreWindow.cpp
+++ b/Engine/Source/Developer/TraceInsightsFrontend/Private/InsightsFrontend/Widgets/STraceStoreWindow.cpp
@@ -2749,8 +2749,11 @@ FReply STraceStoreWindow::AddWatchDir_Clicked()
 		{
 			FPaths::MakePlatformFilename(SelectedDirectory);
 
+			UE_LOG(LogInsightsFrontend, Log, TEXT("[TraceStore] Adding monitored directory: \"%s\"..."), *SelectedDirectory);
+
 			UE::Trace::FStoreClient* StoreClient = TraceStoreConnection->GetStoreClient();
-			if (!StoreClient || !StoreClient->SetStoreDirectories(nullptr, { (*SelectedDirectory) }, {}))
+			if (!StoreClient ||
+				!StoreClient->SetStoreDirectories(nullptr, { (*SelectedDirectory) }, {}))
 			{
 				FMessageLog(LogListingName).Error(LOCTEXT("StoreCommunicationFail", "Failed to change settings on the store service."));
 			}
