Fix persistent settings in Unreal Insights Frontend

diff --git a/Engine/Programs/UnrealInsights/Config/DefaultUnrealInsightsFrontendSettings.ini b/Engine/Programs/UnrealInsights/Config/DefaultUnrealInsightsFrontendSettings.ini
new file mode 100644
--- /dev/null
+++ b/Engine/Programs/UnrealInsights/Config/DefaultUnrealInsightsFrontendSettings.ini
@@ -0,0 +1,12 @@
+[SectionsToSave]
+bCanSaveAllSections=true
+
+[Insights.SessionBrowser]
+; Whether Insights should signal to an UE application to auto-connect and start tracing when Insights is running.
+AutoConnect=True
+; Whether Insights should auto start analysis when a new live trace in received.
+AutoStartAnalysis=False
+; If specified, auto-start analysis will be enabled only for live trace sessions with this specified Platform.
+AutoStartAnalysisPlatform=
+; If specified, auto-start analysis will be enabled only for live trace sessions with this specified App Name.
+AutoStartAnalysisAppName=
\ No newline at end of file
diff --git a/Engine/Programs/UnrealInsights/Config/DefaultUnrealInsightsSessionBrowserSettings.ini b/Engine/Programs/UnrealInsights/Config/DefaultUnrealInsightsSessionBrowserSettings.ini
deleted file mode 100644
diff --git a/Engine/Source/Developer/TraceInsightsFrontend/Private/InsightsFrontend/TraceInsightsFrontendModule.cpp b/Engine/Source/Developer/TraceInsightsFrontend/Private/InsightsFrontend/TraceInsightsFrontendModule.cpp
--- a/Engine/Source/Developer/TraceInsightsFrontend/Private/InsightsFrontend/TraceInsightsFrontendModule.cpp
+++ b/Engine/Source/Developer/TraceInsightsFrontend/Private/InsightsFrontend/TraceInsightsFrontendModule.cpp
@@ -332,6 +332,18 @@ void FTraceInsightsFrontendModule::OnConnectionTabClosed(TSharedRef<SDockTab> Ta
 
 ////////////////////////////////////////////////////////////////////////////////////////////////////
 
+FInsightsFrontendSettings& FTraceInsightsFrontendModule::GetSettings()
+{
+	if (!Settings.IsValid())
+	{
+		Settings = MakeShared<FInsightsFrontendSettings>();
+	}
+
+	return *Settings;
+}
+
+////////////////////////////////////////////////////////////////////////////////////////////////////
+
 void FTraceInsightsFrontendModule::RunAutomationTests(const FString& InCmd)
 {
 #if INSIGHTS_ENABLE_AUTOMATION
diff --git a/Engine/Source/Developer/TraceInsightsFrontend/Private/InsightsFrontend/TraceInsightsFrontendModule.h b/Engine/Source/Developer/TraceInsightsFrontend/Private/InsightsFrontend/TraceInsightsFrontendModule.h
--- a/Engine/Source/Developer/TraceInsightsFrontend/Private/InsightsFrontend/TraceInsightsFrontendModule.h
+++ b/Engine/Source/Developer/TraceInsightsFrontend/Private/InsightsFrontend/TraceInsightsFrontendModule.h
@@ -50,7 +50,7 @@ public:
 
 	//////////////////////////////////////////////////
 
-	FInsightsFrontendSettings& GetSettings() { return Settings; }
+	FInsightsFrontendSettings& GetSettings();
 
 private:
 	void RegisterTabSpawners();
@@ -70,7 +70,7 @@ private:
 
 private:
 	/** An instance of the main settings. */
-	FInsightsFrontendSettings Settings;
+	TSharedPtr<FInsightsFrontendSettings> Settings;
 
 	FCreateFrontendWindowParams CreateWindowParams;
 
