[Insights] Unreal Insights Frontend: Added early memory cleanup for trace metadata analysis. Also made some minor fixes re usage of TWeakPtr.

#rnx

diff --git a/Engine/Source/Developer/TraceInsightsFrontend/Private/InsightsFrontend/StoreService/StoreBrowser.cpp b/Engine/Source/Developer/TraceInsightsFrontend/Private/InsightsFrontend/StoreService/StoreBrowser.cpp
--- a/Engine/Source/Developer/TraceInsightsFrontend/Private/InsightsFrontend/StoreService/StoreBrowser.cpp
+++ b/Engine/Source/Developer/TraceInsightsFrontend/Private/InsightsFrontend/StoreService/StoreBrowser.cpp
@@ -604,7 +604,7 @@ void FStoreBrowser::UpdateMetadata(TSharedPtr<FStoreBrowserTraceInfo> TraceInfoP
 
 		UE::Trace::IInDataStream* Inner;
 		FStoreBrowserTraceInfo* Trace;
-		double TimeLimit = 1.0;
+		double TimeLimit = 1.0; // [seconds]
 		UE::Insights::FStopwatch Stopwatch;
 		int32 BytesRead = 0;
 		EReadStatus Status = EReadStatus::Ready;
@@ -621,6 +621,9 @@ void FStoreBrowser::UpdateMetadata(TSharedPtr<FStoreBrowserTraceInfo> TraceInfoP
 	Context.AddAnalyzer(Analyzer);
 	Context.Process(DataStream).Wait();
 
+	TraceData.Reset();
+	DataStream.Inner = nullptr;
+
 	// Update the FStoreBrowserTraceInfo object.
 	if (Trace.MetadataUpdateCount != 0)
 	{
diff --git a/Engine/Source/Developer/TraceInsightsFrontend/Private/InsightsFrontend/Widgets/STraceListRow.cpp b/Engine/Source/Developer/TraceInsightsFrontend/Private/InsightsFrontend/Widgets/STraceListRow.cpp
--- a/Engine/Source/Developer/TraceInsightsFrontend/Private/InsightsFrontend/Widgets/STraceListRow.cpp
+++ b/Engine/Source/Developer/TraceInsightsFrontend/Private/InsightsFrontend/Widgets/STraceListRow.cpp
@@ -41,8 +41,8 @@ const FName FTraceListColumns::Status(TEXT("Status"));
 
 void STraceListRow::Construct(const FArguments& InArgs, TSharedPtr<FTraceViewModel> InTrace, TSharedRef<STraceStoreWindow> InParentWidget, const TSharedRef<STableViewBase>& InOwnerTableView)
 {
-	WeakTrace = MoveTemp(InTrace);
-	WeakParentWidget = InParentWidget;
+	WeakTrace = InTrace.ToWeakPtr();
+	WeakParentWidget = InParentWidget.ToWeakPtr();
 
 	SMultiColumnTableRow<TSharedPtr<FTraceViewModel>>::Construct(FSuperRowType::FArguments(), InOwnerTableView);
 }
