[Insights]
* Frontend: Added UI to launch Unreal Insights in "Direct Trace" mode (i.e. listening on a certain port for an incoming trace stream to analyze).
* Frontend: Updated Connection UI (trace source, trace destination and tooltips).

diff --git a/Engine/Source/Developer/TraceInsightsFrontend/Private/InsightsFrontend/Widgets/SConnectionWindow.cpp b/Engine/Source/Developer/TraceInsightsFrontend/Private/InsightsFrontend/Widgets/SConnectionWindow.cpp
--- a/Engine/Source/Developer/TraceInsightsFrontend/Private/InsightsFrontend/Widgets/SConnectionWindow.cpp
+++ b/Engine/Source/Developer/TraceInsightsFrontend/Private/InsightsFrontend/Widgets/SConnectionWindow.cpp
@@ -8,9 +8,11 @@
 #include "SocketSubsystem.h"
 #include "Styling/StyleColors.h"
 #include "Widgets/Input/SButton.h"
+#include "Widgets/Input/SCheckBox.h"
 #include "Widgets/Input/SEditableTextBox.h"
 #include "Widgets/Layout/SBorder.h"
 #include "Widgets/Layout/SBox.h"
+#include "Widgets/Layout/SSeparator.h"
 #include "Widgets/Notifications/SNotificationList.h"
 #include "Widgets/SBoxPanel.h"
 #include "Widgets/SOverlay.h"
@@ -20,6 +22,9 @@
 #include "Trace/ControlClient.h"
 #include "Trace/StoreConnection.h"
 
+// TraceInsightsCore
+#include "InsightsCore/Common/MiscUtils.h"
+
 // TraceInsightsFrontend
 #include "InsightsFrontend/Common/InsightsFrontendStyle.h"
 #include "InsightsFrontend/Common/Log.h"
@@ -117,18 +122,236 @@ TSharedRef<SWidget> SConnectionWindow::ConstructConnectPanel()
 			.VAlign(VAlign_Center)
 			[
 				SNew(SBox)
-				.MinDesiredWidth(180.0f)
+				.MinDesiredWidth(200.0f)
 				.HAlign(HAlign_Right)
 				[
 					SNew(STextBlock)
-					.Text(LOCTEXT("TraceRecorderAddressText", "Trace recorder IP address"))
+					.Text(LOCTEXT("StartUnrealInsightsParams_Text", "Direct Trace CmdLine Params"))
 				]
 			]
 
 			+ SHorizontalBox::Slot()
-			.FillWidth(1.0f)
+			.AutoWidth()
+			.VAlign(VAlign_Center)
+			.Padding(8.0f, 0.0f, 0.0f, 0.0f)
+			[
+				SNew(SCheckBox)
+				.IsEnabled(false)
+				.IsChecked(ECheckBoxState::Checked)
+				[
+					SNew(STextBlock)
+					.Text(LOCTEXT("ListenForDirectTrace_Text", "-ListenForDirectTrace"))
+				]
+			]
+
+			+ SHorizontalBox::Slot()
+			.AutoWidth()
+			.VAlign(VAlign_Center)
+			.Padding(8.0f, 0.0f, 0.0f, 0.0f)
+			[
+				SNew(SCheckBox)
+				.IsEnabled(false)
+				.IsChecked(ECheckBoxState::Checked)
+				[
+					SNew(STextBlock)
+					.Text(LOCTEXT("DirectTracePort_Text", "-DirectTracePort="))
+				]
+			]
+
+			+ SHorizontalBox::Slot()
+			.AutoWidth()
+			.VAlign(VAlign_Center)
+			.Padding(2.0f, 0.0f, 0.0f, 0.0f)
+			[
+				SAssignNew(DirectTracePortTextBox, SEditableTextBox)
+				.MinDesiredWidth(60.0f)
+				.RevertTextOnEscape(true)
+				.SelectAllTextWhenFocused(true)
+				.ToolTipText(LOCTEXT("DirectTracePort_Tooltip", "Port to be used by Unreal Insights to listen for a direct trace stream."))
+			]
+		]
+
+		+ SVerticalBox::Slot()
+		.AutoHeight()
+		.HAlign(HAlign_Fill)
+		.Padding(12.0f, 8.0f, 12.0f, 0.0f)
+		[
+			SNew(SHorizontalBox)
+
+			+ SHorizontalBox::Slot()
+			.AutoWidth()
+			.VAlign(VAlign_Center)
+			[
+				SNew(SBox)
+				.MinDesiredWidth(200.0f)
+				.HAlign(HAlign_Right)
+				[
+					SNew(STextBlock)
+					.Text(LOCTEXT("StartUnrealInsightsParams_Text", "Misc CmdLine Params"))
+				]
+			]
+
+
+			+ SHorizontalBox::Slot()
+			.AutoWidth()
+			.VAlign(VAlign_Center)
+			.Padding(8.0f, 0.0f, 0.0f, 0.0f)
+			[
+				SAssignNew(AutoQuitCheckBox, SCheckBox)
+				[
+					SNew(STextBlock)
+					.Text(LOCTEXT("AutoQuit_Text", "-AutoQuit"))
+				]
+			]
+
+			+ SHorizontalBox::Slot()
+			.AutoWidth()
+			.VAlign(VAlign_Center)
+			.Padding(8.0f, 0.0f, 0.0f, 0.0f)
+			[
+				SAssignNew(WaitForSymbolResolverCheckBox, SCheckBox)
+				[
+					SNew(STextBlock)
+					.Text(LOCTEXT("WaitForSymbolResolver_Text", "-WaitForSymbolResolver"))
+				]
+			]
+
+			+ SHorizontalBox::Slot()
+			.AutoWidth()
+			.VAlign(VAlign_Center)
+			.Padding(8.0f, 0.0f, 0.0f, 0.0f)
+			[
+				SAssignNew(DisableFramerateThrottleCheckBox, SCheckBox)
+				[
+					SNew(STextBlock)
+					.Text(LOCTEXT("DisableFramerateThrottle_Text", "-DisableFramerateThrottle"))
+				]
+			]
+
+			+ SHorizontalBox::Slot()
+			.AutoWidth()
 			.VAlign(VAlign_Center)
 			.Padding(6.0f, 0.0f, 0.0f, 0.0f)
+			[
+				SAssignNew(InsightsTestCheckBox, SCheckBox)
+				[
+					SNew(STextBlock)
+					.Text(LOCTEXT("InsightsTest_Text", "-InsightsTest"))
+				]
+			]
+
+			+ SHorizontalBox::Slot()
+			.AutoWidth()
+			.VAlign(VAlign_Center)
+			.Padding(8.0f, 0.0f, 0.0f, 0.0f)
+			[
+				SAssignNew(NoUICheckBox, SCheckBox)
+				[
+					SNew(STextBlock)
+					.Text(LOCTEXT("NoUI_Text", "-NoUI"))
+				]
+			]
+
+			+ SHorizontalBox::Slot()
+			.AutoWidth()
+			.VAlign(VAlign_Center)
+			.Padding(8.0f, 0.0f, 0.0f, 0.0f)
+			[
+				SAssignNew(DebugToolsCheckBox, SCheckBox)
+				[
+					SNew(STextBlock)
+					.Text(LOCTEXT("DebugTools_Text", "-DebugTools"))
+				]
+			]
+		]
+
+		+ SVerticalBox::Slot()
+		.AutoHeight()
+		.HAlign(HAlign_Fill)
+		.Padding(12.0f, 8.0f, 12.0f, 0.0f)
+		[
+			SNew(SHorizontalBox)
+
+			+ SHorizontalBox::Slot()
+			.AutoWidth()
+			.VAlign(VAlign_Center)
+			[
+				SNew(SBox)
+				.MinDesiredWidth(200.0f)
+				.HAlign(HAlign_Right)
+				[
+					SNew(STextBlock)
+					.Text(LOCTEXT("AdditionalParams_Text", "Additional CmdLine Params"))
+				]
+			]
+
+			+ SHorizontalBox::Slot()
+			.FillWidth(1.0f)
+			.VAlign(VAlign_Center)
+			.Padding(8.0f, 0.0f, 0.0f, 0.0f)
+			[
+				SNew(SBox)
+				.MinDesiredWidth(120.0f)
+				[
+					SAssignNew(DirectTraceAdditionalParamsTextBox, SEditableTextBox)
+					.RevertTextOnEscape(true)
+					.SelectAllTextWhenFocused(true)
+					.ToolTipText(LOCTEXT("DirectTraceAdditionalParams_Tooltip", "Additional command line parameter for launching Unreal Insights (Direct Trace mode)."))
+					.MinDesiredWidth(40.0f)
+				]
+			]
+		]
+
+		+ SVerticalBox::Slot()
+		.AutoHeight()
+		.HAlign(HAlign_Fill)
+		.Padding(12.0f, 12.0f, 12.0f, 0.0f)
+		[
+			SNew(SBox)
+			.HAlign(HAlign_Right)
+			.VAlign(VAlign_Center)
+			[
+				SNew(SButton)
+				.ButtonStyle(&FAppStyle::Get().GetWidgetStyle<FButtonStyle>("PrimaryButton"))
+				.Text(LOCTEXT("ListenForDirectTrace", "Start Unreal Insights"))
+				.ToolTipText(LOCTEXT("ListenForDirectTrace_ToolTip", "Start Unreal Insights in \"listen for a Direct Trace stream\" mode."))
+				.OnClicked(this, &SConnectionWindow::StartUnrealInsightsDirectTrace_OnClicked)
+			]
+		]
+
+		+ SVerticalBox::Slot()
+		.AutoHeight()
+		.HAlign(HAlign_Fill)
+		.Padding(12.0f, 12.0f, 12.0f, 0.0f)
+		[
+			SNew(SSeparator)
+		]
+
+		+ SVerticalBox::Slot()
+		.AutoHeight()
+		.HAlign(HAlign_Fill)
+		.Padding(12.0f, 12.0f, 12.0f, 0.0f)
+		[
+			SNew(SHorizontalBox)
+
+			+ SHorizontalBox::Slot()
+			.AutoWidth()
+			.VAlign(VAlign_Center)
+			[
+				SNew(SBox)
+				.MinDesiredWidth(200.0f)
+				.HAlign(HAlign_Right)
+				[
+					SNew(STextBlock)
+					.Text(LOCTEXT("TraceRecorderAddressText", "Trace Destination"))
+					.ToolTipText(LOCTEXT("TraceRecorderAddress_ToolTip", "Destination of the trace stream (trace recorder)\n\nThe IP address of the machine running UnrealTraceServer\nor the IP address of the machine running Unreal Insights in Direct Trace mode."))
+				]
+			]
+
+			+ SHorizontalBox::Slot()
+			.FillWidth(1.0f)
+			.VAlign(VAlign_Center)
+			.Padding(8.0f, 0.0f, 0.0f, 0.0f)
 			[
 				SNew(SBox)
 				.MinDesiredWidth(120.0f)
@@ -150,18 +373,19 @@ TSharedRef<SWidget> SConnectionWindow::ConstructConnectPanel()
 			.VAlign(VAlign_Center)
 			[
 				SNew(SBox)
-				.MinDesiredWidth(180.0f)
+				.MinDesiredWidth(200.0f)
 				.HAlign(HAlign_Right)
 				[
 					SNew(STextBlock)
-					.Text(LOCTEXT("RunningInstanceAddressText", "Running instance IP address"))
+					.Text(LOCTEXT("RunningInstanceAddressText", "Trace Source"))
+					.ToolTipText(LOCTEXT("TraceRecorderAddress_ToolTip", "Source of the trace stream (running instance)\n\nThe IP address of the machine/console running the Editor, game client/server or any UE application having trace enabled."))
 				]
 			]
 
 			+ SHorizontalBox::Slot()
 			.FillWidth(1.0f)
 			.VAlign(VAlign_Center)
-			.Padding(6.0f, 0.0f, 0.0f, 0.0f)
+			.Padding(8.0f, 0.0f, 0.0f, 0.0f)
 			[
 				SNew(SBox)
 				.MinDesiredWidth(120.0f)
@@ -183,18 +407,18 @@ TSharedRef<SWidget> SConnectionWindow::ConstructConnectPanel()
 			.VAlign(VAlign_Center)
 			[
 				SNew(SBox)
-				.MinDesiredWidth(180.0f)
+				.MinDesiredWidth(200.0f)
 				.HAlign(HAlign_Right)
 				[
 					SNew(STextBlock)
-					.Text(LOCTEXT("InitialChannelsText", "Initial channels"))
+					.Text(LOCTEXT("InitialChannelsText", "Initial Channels"))
 				]
 			]
 
 			+ SHorizontalBox::Slot()
 			.FillWidth(1.0f)
 			.VAlign(VAlign_Center)
-			.Padding(6.0f, 0.0f, 0.0f, 0.0f)
+			.Padding(8.0f, 0.0f, 0.0f, 0.0f)
 			[
 				SNew(SBox)
 				.MinDesiredWidth(120.0f)
@@ -207,7 +431,7 @@ TSharedRef<SWidget> SConnectionWindow::ConstructConnectPanel()
 		+ SVerticalBox::Slot()
 		.AutoHeight()
 		.HAlign(HAlign_Fill)
-		.Padding(198.0f, 4.0f, 12.0f, 0.0f)
+		.Padding(220.0f, 4.0f, 12.0f, 0.0f)
 		[
 			SNew(STextBlock)
 			.Text(LOCTEXT("InitialChannelsNoteText", "Comma-separated list of channels/presets to enable when connected."))
@@ -215,7 +439,7 @@ TSharedRef<SWidget> SConnectionWindow::ConstructConnectPanel()
 		+ SVerticalBox::Slot()
 		.AutoHeight()
 		.HAlign(HAlign_Fill)
-		.Padding(198.0f, 2.0f, 12.0f, 0.0f)
+		.Padding(220.0f, 2.0f, 12.0f, 0.0f)
 		[
 			SNew(SHorizontalBox)
 			+ SHorizontalBox::Slot()
@@ -237,7 +461,7 @@ TSharedRef<SWidget> SConnectionWindow::ConstructConnectPanel()
 		+ SVerticalBox::Slot()
 		.AutoHeight()
 		.HAlign(HAlign_Fill)
-		.Padding(198.0f, 2.0f, 12.0f, 0.0f)
+		.Padding(220.0f, 2.0f, 12.0f, 0.0f)
 		[
 			SNew(STextBlock)
 			.Text(LOCTEXT("InitialChannelsNote2Text", "Some channels/presets (like \"memory\") cannot be enabled on late connections."))
@@ -254,7 +478,7 @@ TSharedRef<SWidget> SConnectionWindow::ConstructConnectPanel()
 				SNew(SButton)
 				.ButtonStyle(&FAppStyle::Get().GetWidgetStyle<FButtonStyle>("PrimaryButton"))
 				.Text(LOCTEXT("Connect", "Connect"))
-				.ToolTipText(LOCTEXT("ConnectToolTip", "Late connect the running instance at specified address with the trace recorder."))
+				.ToolTipText(LOCTEXT("ConnectToolTip", "Late connect the running instance (the source of trace data) with the trace recorder (the destination of trace data)."))
 				.OnClicked(this, &SConnectionWindow::Connect_OnClicked)
 				.IsEnabled_Lambda([this]() { return !bIsConnecting; })
 			]
@@ -270,6 +494,9 @@ TSharedRef<SWidget> SConnectionWindow::ConstructConnectPanel()
 			SAssignNew(NotificationList, SNotificationList)
 		];
 
+	// Set the default port for direct trace.
+	DirectTracePortTextBox->SetText(FText::FromString(TEXT("1986")));
+
 	const FText LocalHost = FText::FromString(TEXT("127.0.0.1"));
 
 	if (TraceStoreConnection->IsLocalHost())
@@ -314,6 +541,57 @@ END_SLATE_FUNCTION_BUILD_OPTIMIZATION
 
 ////////////////////////////////////////////////////////////////////////////////////////////////////
 
+FReply SConnectionWindow::StartUnrealInsightsDirectTrace_OnClicked()
+{
+	if (!DirectTracePortTextBox.IsValid())
+	{
+		return FReply::Handled();
+	}
+
+	const uint32 Port = FCString::Atoi(*DirectTracePortTextBox->GetText().ToString());
+
+	FString CmdLine = FString::Printf(TEXT("-ListenForDirectTrace -DirectTracePort=%u"), Port);
+
+	if (AutoQuitCheckBox->IsChecked())
+	{
+		CmdLine.Append(TEXT(" -AutoQuit"));
+	}
+	if (WaitForSymbolResolverCheckBox->IsChecked())
+	{
+		CmdLine.Append(TEXT(" -WaitForSymbolResolver"));
+	}
+	if (DisableFramerateThrottleCheckBox->IsChecked())
+	{
+		CmdLine.Append(TEXT(" -DisableFramerateThrottle"));
+	};
+	if (InsightsTestCheckBox->IsChecked())
+	{
+		CmdLine.Append(TEXT(" -InsightsTest"));
+	}
+	if (NoUICheckBox->IsChecked())
+	{
+		CmdLine.Append(TEXT(" -NoUI"));
+	}
+	if (DebugToolsCheckBox->IsChecked())
+	{
+		CmdLine.Append(TEXT(" -DebugTools"));
+	}
+
+	FText AdditionalParams = DirectTraceAdditionalParamsTextBox->GetText();
+	if (!AdditionalParams.IsEmpty())
+	{
+		CmdLine.Append(TEXT(" "));
+		CmdLine.Append(AdditionalParams.ToString());
+	}
+
+	UE_LOG(LogInsightsFrontend, Log, TEXT("[TraceStore] Start analysis (in separate process) for a direct trace (listen on port %u)"), Port);
+	FMiscUtils::OpenUnrealInsights(*CmdLine);
+
+	return FReply::Handled();
+}
+
+////////////////////////////////////////////////////////////////////////////////////////////////////
+
 FReply SConnectionWindow::Connect_OnClicked()
 {
 	FText TraceRecorderAddressText = TraceRecorderAddressTextBox->GetText();
diff --git a/Engine/Source/Developer/TraceInsightsFrontend/Private/InsightsFrontend/Widgets/SConnectionWindow.h b/Engine/Source/Developer/TraceInsightsFrontend/Private/InsightsFrontend/Widgets/SConnectionWindow.h
--- a/Engine/Source/Developer/TraceInsightsFrontend/Private/InsightsFrontend/Widgets/SConnectionWindow.h
+++ b/Engine/Source/Developer/TraceInsightsFrontend/Private/InsightsFrontend/Widgets/SConnectionWindow.h
@@ -12,6 +12,7 @@
 #include <atomic>
 
 class SVerticalBox;
+class SCheckBox;
 class SEditableTextBox;
 class SNotificationList;
 
@@ -38,15 +39,25 @@ public:
 
 private:
 	TSharedRef<SWidget> ConstructConnectPanel();
+	FReply StartUnrealInsightsDirectTrace_OnClicked();
 	FReply Connect_OnClicked();
 
 private:
 	TSharedPtr<UE::Trace::FStoreConnection> TraceStoreConnection;
 
+	TSharedPtr<SEditableTextBox> DirectTracePortTextBox;
+	TSharedPtr<SEditableTextBox> DirectTraceAdditionalParamsTextBox;
 	TSharedPtr<SEditableTextBox> TraceRecorderAddressTextBox;
 	TSharedPtr<SEditableTextBox> RunningInstanceAddressTextBox;
 	TSharedPtr<SEditableTextBox> ChannelsTextBox;
 
+	TSharedPtr<SCheckBox> AutoQuitCheckBox;
+	TSharedPtr<SCheckBox> WaitForSymbolResolverCheckBox;
+	TSharedPtr<SCheckBox> DisableFramerateThrottleCheckBox;
+	TSharedPtr<SCheckBox> InsightsTestCheckBox;
+	TSharedPtr<SCheckBox> NoUICheckBox;
+	TSharedPtr<SCheckBox> DebugToolsCheckBox;
+
 	/** Widget for the non-intrusive notifications. */
 	TSharedPtr<SNotificationList> NotificationList;
 
