[Insights] Fixed duplicated LOCTEXT.

#rnx

diff --git a/Engine/Source/Developer/TraceInsightsFrontend/Private/InsightsFrontend/Widgets/SConnectionWindow.cpp b/Engine/Source/Developer/TraceInsightsFrontend/Private/InsightsFrontend/Widgets/SConnectionWindow.cpp
--- a/Engine/Source/Developer/TraceInsightsFrontend/Private/InsightsFrontend/Widgets/SConnectionWindow.cpp
+++ b/Engine/Source/Developer/TraceInsightsFrontend/Private/InsightsFrontend/Widgets/SConnectionWindow.cpp
@@ -126,7 +126,7 @@ TSharedRef<SWidget> SConnectionWindow::ConstructConnectPanel()
 				.HAlign(HAlign_Right)
 				[
 					SNew(STextBlock)
-					.Text(LOCTEXT("StartUnrealInsightsParams_Text", "Direct Trace CmdLine Params"))
+					.Text(LOCTEXT("DirectTraceCmdLineParams_Text", "Direct Trace CmdLine Params"))
 				]
 			]
 
@@ -187,11 +187,10 @@ TSharedRef<SWidget> SConnectionWindow::ConstructConnectPanel()
 				.HAlign(HAlign_Right)
 				[
 					SNew(STextBlock)
-					.Text(LOCTEXT("StartUnrealInsightsParams_Text", "Misc CmdLine Params"))
+					.Text(LOCTEXT("MiscCmdLineParams_Text", "Misc CmdLine Params"))
 				]
 			]
 
-
 			+ SHorizontalBox::Slot()
 			.AutoWidth()
 			.VAlign(VAlign_Center)
@@ -281,7 +280,7 @@ TSharedRef<SWidget> SConnectionWindow::ConstructConnectPanel()
 				.HAlign(HAlign_Right)
 				[
 					SNew(STextBlock)
-					.Text(LOCTEXT("AdditionalParams_Text", "Additional CmdLine Params"))
+					.Text(LOCTEXT("AdditionalCmdLineParams_Text", "Additional CmdLine Params"))
 				]
 			]
 
@@ -313,8 +312,8 @@ TSharedRef<SWidget> SConnectionWindow::ConstructConnectPanel()
 			[
 				SNew(SButton)
 				.ButtonStyle(&FAppStyle::Get().GetWidgetStyle<FButtonStyle>("PrimaryButton"))
-				.Text(LOCTEXT("ListenForDirectTrace", "Start Unreal Insights"))
-				.ToolTipText(LOCTEXT("ListenForDirectTrace_ToolTip", "Start Unreal Insights in \"listen for a Direct Trace stream\" mode."))
+				.Text(LOCTEXT("StartUnrealInsights", "Start Unreal Insights"))
+				.ToolTipText(LOCTEXT("StartUnrealInsights_ToolTip", "Start Unreal Insights in \"listen for a Direct Trace stream\" mode."))
 				.OnClicked(this, &SConnectionWindow::StartUnrealInsightsDirectTrace_OnClicked)
 			]
 		]
@@ -343,8 +342,8 @@ TSharedRef<SWidget> SConnectionWindow::ConstructConnectPanel()
 				.HAlign(HAlign_Right)
 				[
 					SNew(STextBlock)
-					.Text(LOCTEXT("TraceRecorderAddressText", "Trace Destination"))
-					.ToolTipText(LOCTEXT("TraceRecorderAddress_ToolTip", "Destination of the trace stream (trace recorder)\n\nThe IP address of the machine running UnrealTraceServer\nor the IP address of the machine running Unreal Insights in Direct Trace mode."))
+					.Text(LOCTEXT("TraceDestination_Text", "Trace Destination"))
+					.ToolTipText(LOCTEXT("TraceDestination_ToolTip", "Destination of the trace stream (i.e. the trace recorder/processor)\n\nThe IP address of the machine running UnrealTraceServer\nor the IP address of the machine running Unreal Insights in Direct Trace mode."))
 				]
 			]
 
@@ -377,8 +376,8 @@ TSharedRef<SWidget> SConnectionWindow::ConstructConnectPanel()
 				.HAlign(HAlign_Right)
 				[
 					SNew(STextBlock)
-					.Text(LOCTEXT("RunningInstanceAddressText", "Trace Source"))
-					.ToolTipText(LOCTEXT("TraceRecorderAddress_ToolTip", "Source of the trace stream (running instance)\n\nThe IP address of the machine/console running the Editor, game client/server or any UE application having trace enabled."))
+					.Text(LOCTEXT("TraceSource_Text", "Trace Source"))
+					.ToolTipText(LOCTEXT("TraceSource_ToolTip", "Source of the trace stream (i.e. the trace emitter)\n\nThe IP address of the machine/console running the editor, game client/server or any UE application having trace enabled."))
 				]
 			]
 
