[Insights]
* Frontend/Connection: Added title and subtitle for "Direct Trace" and "Connection" sections.
* Frontend/Connection: Changed order for "Trace Source" to appear before "Trace Destination".

diff --git a/Engine/Source/Developer/TraceInsightsFrontend/Private/InsightsFrontend/Widgets/SConnectionWindow.cpp b/Engine/Source/Developer/TraceInsightsFrontend/Private/InsightsFrontend/Widgets/SConnectionWindow.cpp
--- a/Engine/Source/Developer/TraceInsightsFrontend/Private/InsightsFrontend/Widgets/SConnectionWindow.cpp
+++ b/Engine/Source/Developer/TraceInsightsFrontend/Private/InsightsFrontend/Widgets/SConnectionWindow.cpp
@@ -110,6 +110,26 @@ TSharedRef<SWidget> SConnectionWindow::ConstructConnectPanel()
 
 	TSharedRef<SWidget> Widget = SNew(SVerticalBox)
 
+		+ SVerticalBox::Slot()
+		.AutoHeight()
+		.HAlign(HAlign_Fill)
+		.Padding(12.0f, 12.0f, 12.0f, 0.0f)
+		[
+			SNew(STextBlock)
+			.Font(FAppStyle::Get().GetFontStyle("Font.Large.Bold"))
+			.Text(LOCTEXT("DirectTrace_Title", "Direct Trace"))
+		]
+
+		+ SVerticalBox::Slot()
+		.AutoHeight()
+		.HAlign(HAlign_Fill)
+		.Padding(12.0f, 4.0f, 12.0f, 0.0f)
+		[
+			SNew(STextBlock)
+			.Text(LOCTEXT("DirectTrace_Subtitle", "Start Unreal Insights in \"listen for a Direct Trace stream\" mode..."))
+			.IsEnabled(false)
+		]
+
 		+ SVerticalBox::Slot()
 		.AutoHeight()
 		.HAlign(HAlign_Fill)
@@ -330,6 +350,25 @@ TSharedRef<SWidget> SConnectionWindow::ConstructConnectPanel()
 		.AutoHeight()
 		.HAlign(HAlign_Fill)
 		.Padding(12.0f, 12.0f, 12.0f, 0.0f)
+		[
+			SNew(STextBlock)
+			.Font(FAppStyle::Get().GetFontStyle("Font.Large.Bold"))
+			.Text(LOCTEXT("Connection_Title", "Connection"))
+		]
+		+ SVerticalBox::Slot()
+		.AutoHeight()
+		.HAlign(HAlign_Fill)
+		.Padding(12.0f, 4.0f, 12.0f, 0.0f)
+		[
+			SNew(STextBlock)
+			.Text(LOCTEXT("Connection_Subtitle", "Late connect the trace source (game client, server, editor or other UE program) with the trace recorder/processor (the destination of trace data)..."))
+			.IsEnabled(false)
+		]
+
+		+ SVerticalBox::Slot()
+		.AutoHeight()
+		.HAlign(HAlign_Fill)
+		.Padding(12.0f, 8.0f, 12.0f, 0.0f)
 		[
 			SNew(SHorizontalBox)
 
@@ -342,8 +381,8 @@ TSharedRef<SWidget> SConnectionWindow::ConstructConnectPanel()
 				.HAlign(HAlign_Right)
 				[
 					SNew(STextBlock)
-					.Text(LOCTEXT("TraceDestination_Text", "Trace Destination"))
-					.ToolTipText(LOCTEXT("TraceDestination_ToolTip", "Destination of the trace stream (i.e. the trace recorder/processor)\n\nThe IP address of the machine running UnrealTraceServer\nor the IP address of the machine running Unreal Insights in Direct Trace mode."))
+					.Text(LOCTEXT("TraceSource_Text", "Trace Source"))
+					.ToolTipText(LOCTEXT("TraceSource_ToolTip", "Source of the trace stream (i.e. the trace emitter)\n\nThe IP address of the machine/console running the editor, game client/server or any UE application having trace enabled."))
 				]
 			]
 
@@ -355,7 +394,7 @@ TSharedRef<SWidget> SConnectionWindow::ConstructConnectPanel()
 				SNew(SBox)
 				.MinDesiredWidth(120.0f)
 				[
-					SAssignNew(TraceRecorderAddressTextBox, SEditableTextBox)
+					SAssignNew(RunningInstanceAddressTextBox, SEditableTextBox)
 				]
 			]
 		]
@@ -363,7 +402,7 @@ TSharedRef<SWidget> SConnectionWindow::ConstructConnectPanel()
 		+ SVerticalBox::Slot()
 		.AutoHeight()
 		.HAlign(HAlign_Fill)
-		.Padding(12.0f, 8.0f, 12.0f, 0.0f)
+		.Padding(12.0f, 12.0f, 12.0f, 0.0f)
 		[
 			SNew(SHorizontalBox)
 
@@ -376,8 +415,8 @@ TSharedRef<SWidget> SConnectionWindow::ConstructConnectPanel()
 				.HAlign(HAlign_Right)
 				[
 					SNew(STextBlock)
-					.Text(LOCTEXT("TraceSource_Text", "Trace Source"))
-					.ToolTipText(LOCTEXT("TraceSource_ToolTip", "Source of the trace stream (i.e. the trace emitter)\n\nThe IP address of the machine/console running the editor, game client/server or any UE application having trace enabled."))
+					.Text(LOCTEXT("TraceDestination_Text", "Trace Destination"))
+					.ToolTipText(LOCTEXT("TraceDestination_ToolTip", "Destination of the trace stream (i.e. the trace recorder/processor)\n\nThe IP address of the machine running UnrealTraceServer or\nthe IP address of the machine running Unreal Insights in Direct Trace mode."))
 				]
 			]
 
@@ -389,7 +428,7 @@ TSharedRef<SWidget> SConnectionWindow::ConstructConnectPanel()
 				SNew(SBox)
 				.MinDesiredWidth(120.0f)
 				[
-					SAssignNew(RunningInstanceAddressTextBox, SEditableTextBox)
+					SAssignNew(TraceRecorderAddressTextBox, SEditableTextBox)
 				]
 			]
 		]
@@ -477,7 +516,7 @@ TSharedRef<SWidget> SConnectionWindow::ConstructConnectPanel()
 				SNew(SButton)
 				.ButtonStyle(&FAppStyle::Get().GetWidgetStyle<FButtonStyle>("PrimaryButton"))
 				.Text(LOCTEXT("Connect", "Connect"))
-				.ToolTipText(LOCTEXT("ConnectToolTip", "Late connect the running instance (the source of trace data) with the trace recorder (the destination of trace data)."))
+				.ToolTipText(LOCTEXT("ConnectToolTip", "Late connect the trace source (game client, server, editor or other UE program)\nwith the trace recorder/processor (the destination of trace data)."))
 				.OnClicked(this, &SConnectionWindow::Connect_OnClicked)
 				.IsEnabled_Lambda([this]() { return !bIsConnecting; })
 			]
