[Insights]
* Session Info: Added separate field for VFSPaths, with paths separated in pairs.
* Session Info: Added support for EngineVersion.
* Unreal Insights Frontend: Added UE (Engine Version) column to the trace sessions table.
* Unreal Insights Frontend: Updated the tooltip for a trace session to show all available info (including EngineVersion and VFSPaths).
* Unreal Insights Frontend: Refactoring re FTraceViewModel.

diff --git a/Engine/Source/Developer/TraceInsights/Private/Insights/Tests/InsightsTestUtils.cpp b/Engine/Source/Developer/TraceInsights/Private/Insights/Tests/InsightsTestUtils.cpp
--- a/Engine/Source/Developer/TraceInsights/Private/Insights/Tests/InsightsTestUtils.cpp
+++ b/Engine/Source/Developer/TraceInsights/Private/Insights/Tests/InsightsTestUtils.cpp
@@ -24,7 +24,6 @@
 #include "Insights/InsightsManager.h"
 #include "Insights/IUnrealInsightsModule.h"
 
-
 ////////////////////////////////////////////////////////////////////////////////////////////////////
 
 FInsightsTestUtils::FInsightsTestUtils(FAutomationTestBase* InTest) :
@@ -399,10 +398,12 @@ bool FInsightsTestUtils::GetMetadata(const FString& TraceFilePath, TMap<FString,
 		OutMetadata.Add("ProjectName", SessionInfo.ProjectName);
 		OutMetadata.Add("Branch", SessionInfo.Branch);
 		OutMetadata.Add("BuildVersion", SessionInfo.BuildVersion);
+		OutMetadata.Add("EngineVersion", SessionInfo.EngineVersion);
 		OutMetadata.Add("Changelist", FString::FromInt(SessionInfo.Changelist));
 		OutMetadata.Add("BuildConfigurationType", LexToString(SessionInfo.ConfigurationType));
 		OutMetadata.Add("BuildTargetType", LexToString(SessionInfo.TargetType));
 		OutMetadata.Add("CommandLine", SessionInfo.CommandLine);
+		OutMetadata.Add("VFSPaths", SessionInfo.VFSPaths);
 
 		return true;
 	}
diff --git a/Engine/Source/Developer/TraceInsights/Private/Insights/Widgets/SSessionInfoWindow.cpp b/Engine/Source/Developer/TraceInsights/Private/Insights/Widgets/SSessionInfoWindow.cpp
--- a/Engine/Source/Developer/TraceInsights/Private/Insights/Widgets/SSessionInfoWindow.cpp
+++ b/Engine/Source/Developer/TraceInsights/Private/Insights/Widgets/SSessionInfoWindow.cpp
@@ -205,7 +205,7 @@ void SSessionInfoWindow::AddInfoLine(TSharedPtr<SVerticalBox> InVerticalBox,
 {
 	InVerticalBox->AddSlot()
 		.AutoHeight()
-		.Padding(16.0f, 0.0f, 16.0f, 0.0f)
+		.Padding(16.0f, 4.0f, 16.0f, 4.0f)
 		[
 			SNew(SHorizontalBox)
 			.Visibility(this, InVisibilityMethodPtr)
@@ -217,7 +217,7 @@ void SSessionInfoWindow::AddInfoLine(TSharedPtr<SVerticalBox> InVerticalBox,
 			.Padding(0.0f)
 			[
 				SNew(SBox)
-				.Padding(FMargin(0.0f, 8.0f, 16.0f, 4.0f))
+				.Padding(FMargin(0.0f, 4.0f, 16.0f, 0.0f))
 				.HAlign(HAlign_Right)
 				.MinDesiredWidth(160.0f)
 				[
@@ -230,7 +230,7 @@ void SSessionInfoWindow::AddInfoLine(TSharedPtr<SVerticalBox> InVerticalBox,
 			.HAlign(HAlign_Fill)
 			.VAlign(VAlign_Top)
 			.FillWidth(1.0f)
-			.Padding(FMargin(0.0f, 4.0f, 0.0f, 4.0f))
+			.Padding(0.0f)
 			[
 				CreateTextBox(TAttribute<FText>(this, InGetTextMethodPtr), bMultiLine)
 			]
@@ -299,12 +299,14 @@ TSharedRef<SDockTab> SSessionInfoWindow::SpawnTab_SessionInfo(const FSpawnTabArg
 	AddInfoLine(VerticalBox, LOCTEXT("Platform_HeaderText",			"Platform"),			&SSessionInfoWindow::GetPlatformText,		&SSessionInfoWindow::IsVisiblePlatformText);
 	AddInfoLine(VerticalBox, LOCTEXT("AppName_HeaderText",			"Application Name"),	&SSessionInfoWindow::GetAppNameText,		&SSessionInfoWindow::IsVisibleAppNameText);
 	AddInfoLine(VerticalBox, LOCTEXT("ProjectName_HeaderText",		"Project Name"),		&SSessionInfoWindow::GetProjectNameText,	&SSessionInfoWindow::IsVisibleProjectNameText);
-	AddInfoLine(VerticalBox, LOCTEXT("Branch_HeaderText",			"Branch"),				&SSessionInfoWindow::GetBranchText,			&SSessionInfoWindow::IsVisibleBranchText);
-	AddInfoLine(VerticalBox, LOCTEXT("BuildVersion_HeaderText",		"Build Version"),		&SSessionInfoWindow::GetBuildVersionText,	&SSessionInfoWindow::IsVisibleBuildVersionText);
-	AddInfoLine(VerticalBox, LOCTEXT("Changelist_HeaderText",		"Changelist"),			&SSessionInfoWindow::GetChangelistText,		&SSessionInfoWindow::IsVisibleChangelistText);
 	AddInfoLine(VerticalBox, LOCTEXT("BuildConfig_HeaderText",		"Build Config"),		&SSessionInfoWindow::GetBuildConfigText,	&SSessionInfoWindow::IsVisibleBuildConfigText);
 	AddInfoLine(VerticalBox, LOCTEXT("BuildTarget_HeaderText",		"Build Target"),		&SSessionInfoWindow::GetBuildTargetText,	&SSessionInfoWindow::IsVisibleBuildTargetText);
+	AddInfoLine(VerticalBox, LOCTEXT("BuildVersion_HeaderText",		"Build Version"),		&SSessionInfoWindow::GetBuildVersionText,	&SSessionInfoWindow::IsVisibleBuildVersionText);
+	AddInfoLine(VerticalBox, LOCTEXT("Branch_HeaderText",			"Branch"),				&SSessionInfoWindow::GetBranchText,			&SSessionInfoWindow::IsVisibleBranchText);
+	AddInfoLine(VerticalBox, LOCTEXT("Changelist_HeaderText",		"Changelist"),			&SSessionInfoWindow::GetChangelistText,		&SSessionInfoWindow::IsVisibleChangelistText);
+	AddInfoLine(VerticalBox, LOCTEXT("EngineVersion_HeaderText",	"Engine Version"),		&SSessionInfoWindow::GetEngineVersionText,	&SSessionInfoWindow::IsVisibleEngineVersionText);
 	AddInfoLine(VerticalBox, LOCTEXT("CommandLine_HeaderText",		"Command Line"),		&SSessionInfoWindow::GetCommandLineText,	&SSessionInfoWindow::IsVisibleCommandLineText, true);
+	AddInfoLine(VerticalBox, LOCTEXT("VFSPaths_HeaderText",			"VFS Paths"),			&SSessionInfoWindow::GetVFSPathsText,		&SSessionInfoWindow::IsVisibleVFSPathsText, true);
 	AddInfoLine(VerticalBox, LOCTEXT("OtherMetadata_HeaderText",	"Other Metadata"),		&SSessionInfoWindow::GetOtherMetadataText,	&SSessionInfoWindow::IsVisibleOtherMetadataText, true);
 	AddInfoLine(VerticalBox, LOCTEXT("StreamType_HeaderText",		"Stream Type"),			&SSessionInfoWindow::GetStreamTypeText,		&SSessionInfoWindow::IsAlwaysVisible);
 
@@ -407,12 +409,14 @@ void SSessionInfoWindow::Tick(const FGeometry& AllottedGeometry, const double In
 		PlatformText = FText::GetEmpty();
 		AppNameText = FText::GetEmpty();
 		ProjectNameText = FText::GetEmpty();
-		BranchText = FText::GetEmpty();
-		BuildVersionText = FText::GetEmpty();
-		ChangelistText = FText::GetEmpty();
 		BuildConfigurationTypeText = FText::GetEmpty();
 		BuildTargetTypeText = FText::GetEmpty();
+		BuildVersionText = FText::GetEmpty();
+		BranchText = FText::GetEmpty();
+		ChangelistText = FText::GetEmpty();
+		EngineVersionText = FText::GetEmpty();
 		CommandLineText = FText::GetEmpty();
+		VFSPathsText = FText::GetEmpty();
 		OtherMetadataText = FText::GetEmpty();
 	}
 
@@ -426,16 +430,54 @@ void SSessionInfoWindow::Tick(const FGeometry& AllottedGeometry, const double In
 		if (DiagnosticsProvider && DiagnosticsProvider->IsSessionInfoAvailable())
 		{
 			TraceServices::FSessionInfo SessionInfo = DiagnosticsProvider->GetSessionInfo();
+
 			PlatformText = FText::FromString(SessionInfo.Platform);
 			AppNameText = FText::FromString(SessionInfo.AppName);
 			ProjectNameText = FText::FromString(SessionInfo.ProjectName);
-			BranchText = FText::FromString(SessionInfo.Branch);
-			BuildVersionText = FText::FromString(SessionInfo.BuildVersion);
-			ChangelistText = FText::AsNumber(SessionInfo.Changelist, &FNumberFormattingOptions::DefaultNoGrouping());
+
 			BuildConfigurationTypeText = FText::FromString(LexToString(SessionInfo.ConfigurationType));
 			BuildTargetTypeText = FText::FromString(LexToString(SessionInfo.TargetType));
+			BuildVersionText = FText::FromString(SessionInfo.BuildVersion);
+			BranchText = FText::FromString(SessionInfo.Branch);
+			ChangelistText = FText::AsNumber(SessionInfo.Changelist, &FNumberFormattingOptions::DefaultNoGrouping());
+
+			EngineVersionText = FText::FromString(SessionInfo.EngineVersion);
+			if (!SessionInfo.EngineVersion.IsEmpty() &&
+				!SessionInfo.BuildVersion.IsEmpty())
+			{
+				// Trim the BuildVersion from the EngineVersion, if it is included and it is the same.
+				int32 Index = SessionInfo.EngineVersion.Find(SessionInfo.BuildVersion, ESearchCase::CaseSensitive);
+				if (Index > 0 && SessionInfo.EngineVersion[Index - 1] == TEXT('-'))
+				{
+					EngineVersionText = FText::FromString(SessionInfo.EngineVersion.Left(Index - 1));
+				}
+			}
+
 			CommandLineText = FText::FromString(SessionInfo.CommandLine);
 
+			TArray<FString> VFSPaths;
+			SessionInfo.VFSPaths.ParseIntoArray(VFSPaths, TEXT(";"), true);
+			int32 NumVFSPaths = 0;
+			TStringBuilder<1024> VFSPathsBuilder;
+			for (FString& Path : VFSPaths)
+			{
+				Path.TrimStartAndEndInline();
+				if (!Path.IsEmpty())
+				{
+					if (NumVFSPaths & 1)
+					{
+						VFSPathsBuilder.Append(TEXT(" \u2192 "));
+					}
+					else if (NumVFSPaths > 0)
+					{
+						VFSPathsBuilder.Append(TEXT("\n"));
+					}
+					++NumVFSPaths;
+					VFSPathsBuilder.Append(Path);
+				}
+			}
+			VFSPathsText = FText::FromString(VFSPathsBuilder.ToString());
+
 			TStringBuilder<1024> OtherMetadata;
 
 			Session->EnumerateMetadata([&OtherMetadata](const TraceServices::FTraceSessionMetadata& Metadata)
@@ -445,12 +487,14 @@ void SSessionInfoWindow::Tick(const FGeometry& AllottedGeometry, const double In
 						FName("Platform"),
 						FName("AppName"),
 						FName("ProjectName"),
-						FName("Branch"),
-						FName("BuildVersion"),
-						FName("Changelist"),
 						FName("ConfigurationType"),
 						FName("TargetType"),
+						FName("BuildVersion"),
+						FName("Branch"),
+						FName("Changelist"),
+						FName("EngineVersion"),
 						FName("CommandLine"),
+						FName("VFSPaths"),
 					};
 
 					for (int32 Index = 0; Index < UE_ARRAY_COUNT(ExcludedMetadata); ++Index)
diff --git a/Engine/Source/Developer/TraceInsights/Private/Insights/Widgets/SSessionInfoWindow.h b/Engine/Source/Developer/TraceInsights/Private/Insights/Widgets/SSessionInfoWindow.h
--- a/Engine/Source/Developer/TraceInsights/Private/Insights/Widgets/SSessionInfoWindow.h
+++ b/Engine/Source/Developer/TraceInsights/Private/Insights/Widgets/SSessionInfoWindow.h
@@ -144,12 +144,14 @@ private:
 	FText GetPlatformText() const { return PlatformText; }
 	FText GetAppNameText() const { return AppNameText; }
 	FText GetProjectNameText() const { return ProjectNameText; }
-	FText GetBranchText() const { return BranchText; }
+	FText GetEngineVersionText() const { return EngineVersionText; }
 	FText GetBuildVersionText() const { return BuildVersionText; }
+	FText GetBranchText() const { return BranchText; }
 	FText GetChangelistText() const { return ChangelistText; }
 	FText GetBuildConfigText() const { return BuildConfigurationTypeText; }
 	FText GetBuildTargetText() const { return BuildTargetTypeText; }
 	FText GetCommandLineText() const { return CommandLineText; }
+	FText GetVFSPathsText() const { return VFSPathsText; }
 	FText GetOtherMetadataText() const { return OtherMetadataText; }
 
 	EVisibility IsAlwaysVisible() const { return EVisibility::Visible; }
@@ -158,12 +160,14 @@ private:
 	EVisibility IsVisiblePlatformText() const { return PlatformText.IsEmpty() ? EVisibility::Collapsed : EVisibility::Visible; }
 	EVisibility IsVisibleAppNameText() const { return AppNameText.IsEmpty() ? EVisibility::Collapsed : EVisibility::Visible; }
 	EVisibility IsVisibleProjectNameText() const { return ProjectNameText.IsEmpty() ? EVisibility::Collapsed : EVisibility::Visible; }
+	EVisibility IsVisibleEngineVersionText() const { return EngineVersionText.IsEmpty() ? EVisibility::Collapsed : EVisibility::Visible; }
 	EVisibility IsVisibleBranchText() const { return BranchText.IsEmpty() ? EVisibility::Collapsed : EVisibility::Visible; }
 	EVisibility IsVisibleBuildVersionText() const { return BuildVersionText.IsEmpty() ? EVisibility::Collapsed : EVisibility::Visible; }
 	EVisibility IsVisibleChangelistText() const { return ChangelistText.IsEmpty() ? EVisibility::Collapsed : EVisibility::Visible; }
 	EVisibility IsVisibleBuildConfigText() const { return BuildConfigurationTypeText.IsEmpty() ? EVisibility::Collapsed : EVisibility::Visible; }
 	EVisibility IsVisibleBuildTargetText() const { return BuildTargetTypeText.IsEmpty() ? EVisibility::Collapsed : EVisibility::Visible; }
 	EVisibility IsVisibleCommandLineText() const { return CommandLineText.IsEmpty() ? EVisibility::Collapsed : EVisibility::Visible; }
+	EVisibility IsVisibleVFSPathsText() const { return VFSPathsText.IsEmpty() ? EVisibility::Collapsed : EVisibility::Visible; }
 	EVisibility IsVisibleOtherMetadataText() const { return OtherMetadataText.IsEmpty() ? EVisibility::Collapsed : EVisibility::Visible; }
 
 	FText GetFileSizeText() const;
@@ -189,12 +193,14 @@ private:
 	FText PlatformText;
 	FText AppNameText;
 	FText ProjectNameText;
-	FText BranchText;
+	FText EngineVersionText;
 	FText BuildVersionText;
+	FText BranchText;
 	FText ChangelistText;
 	FText BuildConfigurationTypeText;
 	FText BuildTargetTypeText;
 	FText CommandLineText;
+	FText VFSPathsText;
 	FText OtherMetadataText;
 
 	TSharedPtr<const TraceServices::IAnalysisSession> AnalysisSession;
diff --git a/Engine/Source/Developer/TraceInsightsFrontend/Private/InsightsFrontend/StoreService/DiagnosticsSessionAnalyzer.cpp b/Engine/Source/Developer/TraceInsightsFrontend/Private/InsightsFrontend/StoreService/DiagnosticsSessionAnalyzer.cpp
--- a/Engine/Source/Developer/TraceInsightsFrontend/Private/InsightsFrontend/StoreService/DiagnosticsSessionAnalyzer.cpp
+++ b/Engine/Source/Developer/TraceInsightsFrontend/Private/InsightsFrontend/StoreService/DiagnosticsSessionAnalyzer.cpp
@@ -57,16 +57,17 @@ bool FDiagnosticsSessionAnalyzer::OnEvent(uint16 RouteId, EStyle, const FOnEvent
 		EventData.GetString("Platform", Platform);
 		EventData.GetString("AppName", AppName);
 		EventData.GetString("ProjectName", ProjectName);
-		EventData.GetString("CommandLine", CommandLine);
-		EventData.GetString("Branch", Branch);
 		EventData.GetString("BuildVersion", BuildVersion);
+		EventData.GetString("Branch", Branch);
+		EventData.GetString("EngineVersion", EngineVersion);
+		EventData.GetString("CommandLine", CommandLine);
+		EventData.GetString("VFSPaths", VFSPaths);
 		Changelist = EventData.GetValue<uint32>("Changelist", 0);
-		ConfigurationType = (EBuildConfiguration) EventData.GetValue<uint8>("ConfigurationType");
-		TargetType = (EBuildTargetType) EventData.GetValue<uint8>("TargetType");
-
+		ConfigurationType = (EBuildConfiguration)EventData.GetValue<uint8>("ConfigurationType");
+		TargetType = (EBuildTargetType)EventData.GetValue<uint8>("TargetType");
 		return false;
-	};
 	}
+	} // switch
 
 	return true;
 }
diff --git a/Engine/Source/Developer/TraceInsightsFrontend/Private/InsightsFrontend/StoreService/DiagnosticsSessionAnalyzer.h b/Engine/Source/Developer/TraceInsightsFrontend/Private/InsightsFrontend/StoreService/DiagnosticsSessionAnalyzer.h
--- a/Engine/Source/Developer/TraceInsightsFrontend/Private/InsightsFrontend/StoreService/DiagnosticsSessionAnalyzer.h
+++ b/Engine/Source/Developer/TraceInsightsFrontend/Private/InsightsFrontend/StoreService/DiagnosticsSessionAnalyzer.h
@@ -24,9 +24,11 @@ struct FDiagnosticsSessionAnalyzer : public UE::Trace::IAnalyzer
 	FString Platform;
 	FString AppName;
 	FString ProjectName;
-	FString CommandLine;
-	FString Branch;
 	FString BuildVersion;
+	FString Branch;
+	FString EngineVersion;
+	FString CommandLine;
+	FString VFSPaths;
 	uint32 Changelist = 0;
 	EBuildConfiguration ConfigurationType = EBuildConfiguration::Unknown;
 	EBuildTargetType TargetType = EBuildTargetType::Unknown;
diff --git a/Engine/Source/Developer/TraceInsightsFrontend/Private/InsightsFrontend/StoreService/StoreBrowser.h b/Engine/Source/Developer/TraceInsightsFrontend/Private/InsightsFrontend/StoreService/StoreBrowser.h
--- a/Engine/Source/Developer/TraceInsightsFrontend/Private/InsightsFrontend/StoreService/StoreBrowser.h
+++ b/Engine/Source/Developer/TraceInsightsFrontend/Private/InsightsFrontend/StoreService/StoreBrowser.h
@@ -35,9 +35,11 @@ struct FStoreBrowserTraceInfo
 	FString Platform;
 	FString AppName;
 	FString ProjectName;
-	FString CommandLine;
-	FString Branch;
 	FString BuildVersion;
+	FString Branch;
+	FString EngineVersion;
+	FString CommandLine;
+	FString VFSPaths;
 	uint32 Changelist;
 	EBuildConfiguration ConfigurationType = EBuildConfiguration::Unknown;
 	EBuildTargetType TargetType = EBuildTargetType::Unknown;
diff --git a/Engine/Source/Developer/TraceInsightsFrontend/Private/InsightsFrontend/ViewModels/TraceSetFilter.cpp b/Engine/Source/Developer/TraceInsightsFrontend/Private/InsightsFrontend/ViewModels/TraceSetFilter.cpp
--- a/Engine/Source/Developer/TraceInsightsFrontend/Private/InsightsFrontend/ViewModels/TraceSetFilter.cpp
+++ b/Engine/Source/Developer/TraceInsightsFrontend/Private/InsightsFrontend/ViewModels/TraceSetFilter.cpp
@@ -149,6 +149,13 @@ FTraceFilterByBuildTarget::FTraceFilterByBuildTarget()
 
 ////////////////////////////////////////////////////////////////////////////////////////////////////
 
+FTraceFilterByBuildVersion::FTraceFilterByBuildVersion()
+{
+	ToggleAllActionTooltip = LOCTEXT("FilterByBuildVersion_ToggleAll_Tooltip", "Shows or hides all build versions.");
+}
+
+////////////////////////////////////////////////////////////////////////////////////////////////////
+
 FTraceFilterByBranch::FTraceFilterByBranch()
 {
 	ToggleAllActionTooltip = LOCTEXT("FilterByBranch_ToggleAll_Tooltip", "Shows or hides traces for all branches.");
@@ -156,6 +163,42 @@ FTraceFilterByBranch::FTraceFilterByBranch()
 
 ////////////////////////////////////////////////////////////////////////////////////////////////////
 
+FTraceFilterByChangelist::FTraceFilterByChangelist()
+{
+	ToggleAllActionTooltip = LOCTEXT("FilterByChangelist_ToggleAll_Tooltip", "Shows or hides traces of all changelists.");
+}
+
+////////////////////////////////////////////////////////////////////////////////////////////////////
+
+void FTraceFilterByChangelist::AddDefaultValues(TArray<uint8>& InOutDefaultValues) const
+{
+	for (uint8 ChangelistCategory = 0; ChangelistCategory < (uint8)EChangelistCategory::InvalidOrMax; ++ChangelistCategory)
+	{
+		InOutDefaultValues.Add(ChangelistCategory);
+	}
+}
+
+////////////////////////////////////////////////////////////////////////////////////////////////////
+
+FText FTraceFilterByChangelist::ValueToText(const uint8 InValue) const
+{
+	switch ((EChangelistCategory)InValue)
+	{
+	case EChangelistCategory::Zero:    return LOCTEXT("FilterByChangelist_Zero", "0");
+	case EChangelistCategory::NotZero: return LOCTEXT("FilterByChangelist_NotZero", "!= 0");
+	default:                           return UndefinedValueLabel;
+	}
+}
+
+////////////////////////////////////////////////////////////////////////////////////////////////////
+
+FTraceFilterByEngineVersion::FTraceFilterByEngineVersion()
+{
+	ToggleAllActionTooltip = LOCTEXT("FilterByEngineVersion_ToggleAll_Tooltip", "Shows or hides all engine versions.");
+}
+
+////////////////////////////////////////////////////////////////////////////////////////////////////
+
 FTraceFilterBySize::FTraceFilterBySize()
 {
 	ToggleAllActionTooltip = LOCTEXT("FilterBySize_ToggleAll_Tooltip", "Shows or hides traces of all sizes.");
@@ -216,13 +259,6 @@ FText FTraceFilterByStatus::ValueToText(const bool InValue) const
 
 ////////////////////////////////////////////////////////////////////////////////////////////////////
 
-FTraceFilterByVersion::FTraceFilterByVersion()
-{
-	ToggleAllActionTooltip = LOCTEXT("FilterByVersion_ToggleAll_Tooltip", "Shows or hides all versions.");
-}
-
-////////////////////////////////////////////////////////////////////////////////////////////////////
-
 } // namespace UE::Insights
 
 #undef LOCTEXT_NAMESPACE
diff --git a/Engine/Source/Developer/TraceInsightsFrontend/Private/InsightsFrontend/ViewModels/TraceSetFilter.h b/Engine/Source/Developer/TraceInsightsFrontend/Private/InsightsFrontend/ViewModels/TraceSetFilter.h
--- a/Engine/Source/Developer/TraceInsightsFrontend/Private/InsightsFrontend/ViewModels/TraceSetFilter.h
+++ b/Engine/Source/Developer/TraceInsightsFrontend/Private/InsightsFrontend/ViewModels/TraceSetFilter.h
@@ -83,7 +83,7 @@ public:
 protected:
 	virtual FString GetFilterValueForTrace(const FTraceViewModel& InTrace) const override
 	{
-		return InTrace.Platform.ToString();
+		return InTrace.GetPlatform();
 	}
 };
 
@@ -97,7 +97,7 @@ public:
 protected:
 	virtual FString GetFilterValueForTrace(const FTraceViewModel& InTrace) const override
 	{
-		return InTrace.AppName.ToString();
+		return InTrace.GetAppName();
 	}
 };
 
@@ -111,7 +111,7 @@ public:
 protected:
 	virtual uint8 GetFilterValueForTrace(const FTraceViewModel& InTrace) const override
 	{
-		return (uint8)InTrace.ConfigurationType;
+		return (uint8)InTrace.GetBuildConfiguration();
 	}
 
 	virtual FText ValueToText(const uint8 InValue) const override
@@ -131,7 +131,7 @@ public:
 protected:
 	virtual uint8 GetFilterValueForTrace(const FTraceViewModel& InTrace) const override
 	{
-		return (uint8)InTrace.TargetType;
+		return (uint8)InTrace.GetBuildTarget();
 	}
 
 	virtual FText ValueToText(const uint8 InValue) const override
@@ -143,6 +143,20 @@ protected:
 
 ////////////////////////////////////////////////////////////////////////////////////////////////////
 
+class FTraceFilterByBuildVersion : public FTraceFilterByStringSet
+{
+public:
+	FTraceFilterByBuildVersion();
+
+protected:
+	virtual FString GetFilterValueForTrace(const FTraceViewModel& InTrace) const override
+	{
+		return InTrace.GetBuildVersion();
+	}
+};
+
+////////////////////////////////////////////////////////////////////////////////////////////////////
+
 class FTraceFilterByBranch : public FTraceFilterByStringSet
 {
 public:
@@ -151,7 +165,55 @@ public:
 protected:
 	virtual FString GetFilterValueForTrace(const FTraceViewModel& InTrace) const override
 	{
-		return InTrace.Branch.ToString();
+		return InTrace.GetBuildBranch();
+	}
+};
+
+////////////////////////////////////////////////////////////////////////////////////////////////////
+
+class FTraceFilterByChangelist : public TTraceSetFilter<uint8>
+{
+public:
+	enum class EChangelistCategory : uint8
+	{
+		Zero,  // CL 0
+		NotZero,
+
+		InvalidOrMax
+	};
+
+public:
+	FTraceFilterByChangelist();
+
+protected:
+	virtual void AddDefaultValues(TArray<uint8>& InOutDefaultValues) const override;
+
+	virtual uint8 GetFilterValueForTrace(const FTraceViewModel& InTrace) const override
+	{
+		if (InTrace.GetChangelist() == 0)
+		{
+			return (uint8)EChangelistCategory::Zero;
+		}
+		else
+		{
+			return (uint8)EChangelistCategory::NotZero;
+		}
+	}
+
+	virtual FText ValueToText(const uint8 InValue) const override;
+};
+
+////////////////////////////////////////////////////////////////////////////////////////////////////
+
+class FTraceFilterByEngineVersion : public FTraceFilterByStringSet
+{
+public:
+	FTraceFilterByEngineVersion();
+
+protected:
+	virtual FString GetFilterValueForTrace(const FTraceViewModel& InTrace) const override
+	{
+		return InTrace.GetEngineVersion();
 	}
 };
 
@@ -179,15 +241,16 @@ protected:
 
 	virtual uint8 GetFilterValueForTrace(const FTraceViewModel& InTrace) const override
 	{
-		if (InTrace.Size == 0)
+		const uint64 TraceSize = InTrace.GetSize();
+		if (TraceSize == 0)
 		{
 			return (uint8)ESizeCategory::Empty;
 		}
-		else if (InTrace.Size < 1024ull * 1024ull)
+		else if (TraceSize < 1024ull * 1024ull)
 		{
 			return (uint8)ESizeCategory::Small;
 		}
-		else if (InTrace.Size < 1024ull * 1024ull * 1024ull)
+		else if (TraceSize < 1024ull * 1024ull * 1024ull)
 		{
 			return (uint8)ESizeCategory::Medium;
 		}
@@ -212,7 +275,7 @@ protected:
 
 	virtual bool GetFilterValueForTrace(const FTraceViewModel& InTrace) const override
 	{
-		return InTrace.bIsLive;
+		return InTrace.IsLive();
 	}
 
 	virtual FText ValueToText(const bool InValue) const override;
@@ -220,18 +283,4 @@ protected:
 
 ////////////////////////////////////////////////////////////////////////////////////////////////////
 
-class FTraceFilterByVersion : public FTraceFilterByStringSet
-{
-public:
-	FTraceFilterByVersion();
-
-protected:
-	virtual FString GetFilterValueForTrace(const FTraceViewModel& InTrace) const override
-	{
-		return InTrace.BuildVersion.ToString();
-	}
-};
-
-////////////////////////////////////////////////////////////////////////////////////////////////////
-
 } // namespace UE::Insights
diff --git a/Engine/Source/Developer/TraceInsightsFrontend/Private/InsightsFrontend/ViewModels/TraceViewModel.cpp b/Engine/Source/Developer/TraceInsightsFrontend/Private/InsightsFrontend/ViewModels/TraceViewModel.cpp
--- a/Engine/Source/Developer/TraceInsightsFrontend/Private/InsightsFrontend/ViewModels/TraceViewModel.cpp
+++ b/Engine/Source/Developer/TraceInsightsFrontend/Private/InsightsFrontend/ViewModels/TraceViewModel.cpp
@@ -13,4 +13,111 @@ FText FTraceViewModel::AnsiStringViewToText(const FAnsiStringView& AnsiStringVie
 	return FText::FromString(FatString);
 }
 
+void FTraceViewModel::SetEngineVersion(const FString& InEngineVersion)
+{
+	FString Version = InEngineVersion;
+
+	if (Version.IsEmpty())
+	{
+		FString BranchString = BranchText.ToString();
+		if (!BranchString.IsEmpty())
+		{
+			FStringView UE4ReleasePrefix = TEXTVIEW("++UE4+Release-");
+			FStringView UE5ReleasePrefix = TEXTVIEW("++UE5+Release-");
+			if (BranchString.StartsWith(UE5ReleasePrefix))
+			{
+				Version = BranchString.RightChop(UE5ReleasePrefix.Len());
+			}
+			else if (BranchString.StartsWith(TEXTVIEW("++UE5+")) ||
+					 BranchString.StartsWith(TEXTVIEW("UE5")))
+			{
+				Version = "5";
+			}
+			else if (BranchString.StartsWith(UE4ReleasePrefix))
+			{
+				Version = BranchString.RightChop(UE4ReleasePrefix.Len());
+			}
+			else if (BranchString.StartsWith(TEXTVIEW("++UE4+")) ||
+					 BranchString.StartsWith(TEXTVIEW("UE4")))
+			{
+				Version = "4";
+			}
+		}
+	}
+	else if (!BuildVersionText.IsEmpty())
+	{
+		// Trim the BuildVersion from the EngineVersion, if it is included and it is the same.
+		int32 Index = Version.Find(BuildVersionText.ToString(), ESearchCase::CaseSensitive);
+		if (Index > 0 && Version[Index - 1] == TEXT('-'))
+		{
+			Version.LeftInline(Index - 1);
+		}
+	}
+
+	EngineVersionText = FText::FromString(Version);
+
+	auto ReadVersionNumber = [](FString& InOutVersion) -> uint64
+		{
+			if (InOutVersion.IsEmpty())
+			{
+				return 0ull;
+			}
+			int32 Index = 0;
+			while (Index < InOutVersion.Len() &&
+				InOutVersion[Index] >= TEXT('0') &&
+				InOutVersion[Index] <= TEXT('9'))
+			{
+				++Index;
+			}
+			FString VersionNumberString = InOutVersion.Left(Index);
+			int32 VersionNumber = FCString::Atoi(*VersionNumberString);
+			if (Index < InOutVersion.Len() && InOutVersion[Index] == TEXT('.'))
+			{
+				InOutVersion.RightChopInline(Index + 1);
+			}
+			else
+			{
+				InOutVersion.Reset();
+			}
+			return uint64(VersionNumber + 1); // +1 is to sort "5.0" after "5"
+		};
+
+	uint64 Major = ReadVersionNumber(Version);
+	uint64 Minor = ReadVersionNumber(Version);
+	uint64 Patch = ReadVersionNumber(Version);
+
+	EngineVersionNumber = (Major << 54) | (Minor << 44) | (Patch << 32) | uint64(Changelist);
+}
+
+void FTraceViewModel::SetCommandLine(const FString& InCommandLine)
+{
+	CommandLineText = FText::FromString(InCommandLine.TrimStartAndEnd());
+}
+
+void FTraceViewModel::SetVFSPaths(const FString& InVFSPaths)
+{
+	TArray<FString> VFSPathsArray;
+	InVFSPaths.ParseIntoArray(VFSPathsArray, TEXT(";"), true);
+	int32 NumVFSPaths = 0;
+	TStringBuilder<1024> VFSPathsBuilder;
+	for (FString& Path : VFSPathsArray)
+	{
+		Path.TrimStartAndEndInline();
+		if (!Path.IsEmpty())
+		{
+			if (NumVFSPaths & 1)
+			{
+				VFSPathsBuilder.Append(TEXT(" \u2192 "));
+			}
+			else if (NumVFSPaths > 0)
+			{
+				VFSPathsBuilder.Append(TEXT("\n"));
+			}
+			++NumVFSPaths;
+			VFSPathsBuilder.Append(Path);
+		}
+	}
+	VFSPathsText = FText::FromString(VFSPathsBuilder.ToString());
+}
+
 } // namespace UE::Insights
diff --git a/Engine/Source/Developer/TraceInsightsFrontend/Private/InsightsFrontend/ViewModels/TraceViewModel.h b/Engine/Source/Developer/TraceInsightsFrontend/Private/InsightsFrontend/ViewModels/TraceViewModel.h
--- a/Engine/Source/Developer/TraceInsightsFrontend/Private/InsightsFrontend/ViewModels/TraceViewModel.h
+++ b/Engine/Source/Developer/TraceInsightsFrontend/Private/InsightsFrontend/ViewModels/TraceViewModel.h
@@ -16,38 +16,92 @@ class SEditableTextBox;
 namespace UE::Insights
 {
 
-struct FTraceViewModel
+class FTraceViewModel
 {
+public:
 	static constexpr uint32 InvalidTraceId = 0;
 
-	uint32 TraceId = InvalidTraceId;
+public:
+	FTraceViewModel() = default;
 
-	uint64 ChangeSerial = 0;
+	uint64 GetChangeSerial() const { return ChangeSerial; }
+	void SetChangeSerial(uint64 InChangeSerial) { ChangeSerial = InChangeSerial; }
 
-	FText Name;
-	FText Uri;
-	FSlateColor DirectoryColor;
+	uint32 GetId() const { return TraceId; }
+	bool HasValidId() const { return TraceId != InvalidTraceId; }
+	void SetId(uint32 InId) { TraceId = InId; }
 
-	FDateTime Timestamp = 0;
-	uint64 Size = 0;
+	uint64 GetSize() const { return Size; }
+	void SetSize(uint64 InSize) { Size = InSize; }
 
-	FText Platform;
-	FText AppName;
-	FText CommandLine;
-	FText Branch;
-	FText BuildVersion;
-	uint32 Changelist = 0;
-	EBuildConfiguration ConfigurationType = EBuildConfiguration::Unknown;
-	EBuildTargetType TargetType = EBuildTargetType::Unknown;
+	FDateTime GetTimestamp() const { return Timestamp; }
+	void SetTimestamp(const FDateTime& InTimestamp) { Timestamp = InTimestamp; }
 
-	bool bIsMetadataUpdated = false;
-	bool bIsRenaming = false;
-	bool bIsLive = false;
-	uint32 IpAddress = 0;
+	const FText& GetNameText() const { return NameText; }
+	const FString& GetName() const { return NameText.ToString(); }
+	void SetName(const FString& InName) { NameText = FText::FromString(InName); }
+	bool HasSameName(const FTraceViewModel& Other) const { return NameText.EqualTo(Other.GetNameText()); }
 
-	TWeakPtr<SEditableTextBox> RenameTextBox;
+	const FText& GetUriText() const { return UriText; }
+	const FString& GetUri() const { return UriText.ToString(); }
+	void SetUri(const FString& InUri) { UriText = FText::FromString(InUri); }
+	bool HasSameUri(const FTraceViewModel& Other) const { return UriText.EqualTo(Other.GetUriText()); }
 
-	FTraceViewModel() = default;
+	const FText& GetPlatformText() const { return PlatformText; }
+	const FString& GetPlatform() const { return PlatformText.ToString(); }
+	void SetPlatform(const FString& InPlatform) { PlatformText = FText::FromString(InPlatform); }
+
+	const FText& GetAppNameText() const { return AppNameText; }
+	const FString& GetAppName() const { return AppNameText.ToString(); }
+	void SetAppName(const FString& InAppName) { AppNameText = FText::FromString(InAppName); }
+
+	EBuildConfiguration GetBuildConfiguration() const { return ConfigurationType; }
+	void SetBuildConfiguration(EBuildConfiguration InBuildConfiguration) { ConfigurationType = InBuildConfiguration; }
+
+	EBuildTargetType GetBuildTarget() const { return TargetType; }
+	void SetBuildTarget(EBuildTargetType InBuildTarget) { TargetType = InBuildTarget; }
+
+	const FText& GetBuildVersionText() const { return BuildVersionText; }
+	const FString& GetBuildVersion() const { return BuildVersionText.ToString(); }
+	void SetBuildVersion(const FString& InBuildVersion) { BuildVersionText = FText::FromString(InBuildVersion); }
+
+	const FText& GetBuildBranchText() const { return BranchText; }
+	const FString& GetBuildBranch() const { return BranchText.ToString(); }
+	void SetBuildBranch(const FString& InBuildBranch) { BranchText = FText::FromString(InBuildBranch); }
+
+	uint32 GetChangelist() const { return Changelist; }
+	void SetChangelist(uint32 InChangelist) { Changelist = InChangelist; }
+
+	const FText& GetEngineVersionText() const { return EngineVersionText; }
+	const FString& GetEngineVersion() const { return EngineVersionText.ToString(); }
+	uint64 GetEngineVersionNumber() const { return EngineVersionNumber; }
+	void SetEngineVersion(const FString& InEngineVersion);
+
+	const FText& GetCommandLineText() const { return CommandLineText; }
+	const FString& GetCommandLine() const { return CommandLineText.ToString(); }
+	void SetCommandLine(const FString& InCommandLine);
+
+	const FText& GetVFSPathsText() const { return VFSPathsText; }
+	const FString& GetVFSPaths() const { return VFSPathsText.ToString(); }
+	void SetVFSPaths(const FString& InVFSPaths);
+
+	bool IsLive() const { return bIsLive; }
+	void SetIsLive(bool bOnOff) { bIsLive = bOnOff; }
+
+	bool IsMetadataUpdated() const { return bIsMetadataUpdated; }
+	void SetIsMetadataUpdated(bool bOnOff) { bIsMetadataUpdated = bOnOff; }
+
+	bool IsRenaming() const { return bIsRenaming; }
+	void StartRenaming() { bIsRenaming = true; }
+	void StopRenaming() { bIsRenaming = false; }
+	TSharedPtr<SEditableTextBox> GetRenameTextBox() const& { return RenameTextBox.Pin(); }
+	void SetRenameTextBox(TWeakPtr<SEditableTextBox> InRenameTextBox) { RenameTextBox = InRenameTextBox; }
+
+	const FSlateColor& GetDirectoryColor() const { return DirectoryColor; }
+	void SetDirectoryColor(const FSlateColor& InColor) { DirectoryColor = InColor; }
+
+	const uint32 GetIpAddress() const { return IpAddress; }
+	void SetIpAddress(uint32 InIpAddress) { IpAddress = InIpAddress; }
 
 	static FDateTime ConvertTimestamp(uint64 InTimestamp)
 	{
@@ -55,6 +109,37 @@ struct FTraceViewModel
 	}
 
 	static FText AnsiStringViewToText(const FAnsiStringView& AnsiStringView);
+
+private:
+	uint64 ChangeSerial = 0;
+	uint64 EngineVersionNumber = 0; // only used for sorting
+	uint32 TraceId = InvalidTraceId;
+	uint32 Changelist = 0;
+	uint64 Size = 0;
+	FDateTime Timestamp = 0;
+
+	FText NameText;
+	FText UriText;
+	FText PlatformText;
+	FText AppNameText;
+	FText BuildVersionText;
+	FText BranchText;
+	FText EngineVersionText;
+	FText CommandLineText;
+	FText VFSPathsText;
+
+	FSlateColor DirectoryColor;
+
+	uint32 IpAddress = 0;
+
+	EBuildConfiguration ConfigurationType = EBuildConfiguration::Unknown;
+	EBuildTargetType TargetType = EBuildTargetType::Unknown;
+
+	bool bIsLive = false;
+	bool bIsMetadataUpdated = false;
+	bool bIsRenaming = false;
+
+	TWeakPtr<SEditableTextBox> RenameTextBox;
 };
 
 } // namespace UE::Insights
diff --git a/Engine/Source/Developer/TraceInsightsFrontend/Private/InsightsFrontend/Widgets/STraceListRow.cpp b/Engine/Source/Developer/TraceInsightsFrontend/Private/InsightsFrontend/Widgets/STraceListRow.cpp
--- a/Engine/Source/Developer/TraceInsightsFrontend/Private/InsightsFrontend/Widgets/STraceListRow.cpp
+++ b/Engine/Source/Developer/TraceInsightsFrontend/Private/InsightsFrontend/Widgets/STraceListRow.cpp
@@ -32,8 +32,10 @@ const FName FTraceListColumns::Platform(TEXT("Platform"));
 const FName FTraceListColumns::AppName(TEXT("AppName"));
 const FName FTraceListColumns::BuildConfig(TEXT("BuildConfig"));
 const FName FTraceListColumns::BuildTarget(TEXT("BuildTarget"));
-const FName FTraceListColumns::BuildBranch(TEXT("BuildBranch"));
 const FName FTraceListColumns::BuildVersion(TEXT("BuildVersion"));
+const FName FTraceListColumns::BuildBranch(TEXT("BuildBranch"));
+const FName FTraceListColumns::Changelist(TEXT("Changelist"));
+const FName FTraceListColumns::EngineVersion(TEXT("EngineVersion"));
 const FName FTraceListColumns::Size(TEXT("Size"));
 const FName FTraceListColumns::Status(TEXT("Status"));
 
@@ -108,7 +110,7 @@ TSharedRef<SWidget> STraceListRow::GenerateWidgetForColumn(const FName& ColumnNa
 		{
 			// A weak reference is stored in the view model in order to focus the text box
 			// when starts renaming.
-			TracePin->RenameTextBox = RenameTextBox;
+			TracePin->SetRenameTextBox(RenameTextBox.ToWeakPtr());
 		}
 
 		return Widget;
@@ -163,6 +165,16 @@ TSharedRef<SWidget> STraceListRow::GenerateWidgetForColumn(const FName& ColumnNa
 				.ToolTip(STraceListRow::GetTraceTooltip())
 			];
 	}
+	else if (ColumnName == FTraceListColumns::BuildVersion)
+	{
+		return SNew(SBox)
+			.Padding(FMargin(4.0f, 0.0f))
+			[
+				SNew(STextBlock)
+				.Text(this, &STraceListRow::GetTraceBuildVersion)
+				.ToolTip(STraceListRow::GetTraceTooltip())
+			];
+	}
 	else if (ColumnName == FTraceListColumns::BuildBranch)
 	{
 		return SNew(SBox)
@@ -173,13 +185,23 @@ TSharedRef<SWidget> STraceListRow::GenerateWidgetForColumn(const FName& ColumnNa
 				.ToolTip(STraceListRow::GetTraceTooltip())
 			];
 	}
-	else if (ColumnName == FTraceListColumns::BuildVersion)
+	else if (ColumnName == FTraceListColumns::Changelist)
 	{
 		return SNew(SBox)
 			.Padding(FMargin(4.0f, 0.0f))
 			[
 				SNew(STextBlock)
-				.Text(this, &STraceListRow::GetTraceBuildVersion)
+				.Text(this, &STraceListRow::GetTraceChangelist)
+				.ToolTip(STraceListRow::GetTraceTooltip())
+			];
+	}
+	else if (ColumnName == FTraceListColumns::EngineVersion)
+	{
+		return SNew(SBox)
+			.Padding(FMargin(4.0f, 0.0f))
+			[
+				SNew(STextBlock)
+				.Text(this, &STraceListRow::GetTraceEngineVersion)
 				.ToolTip(STraceListRow::GetTraceTooltip())
 			];
 	}
@@ -220,7 +242,7 @@ bool STraceListRow::IsRenaming() const
 	TSharedPtr<FTraceViewModel> TracePin = WeakTrace.Pin();
 	if (TracePin.IsValid())
 	{
-		return TracePin->bIsRenaming;
+		return TracePin->IsRenaming();
 	}
 	return false;
 }
@@ -236,7 +258,7 @@ void STraceListRow::RenameTextBox_OnValueCommitted(const FText& InText, ETextCom
 		{
 			Rename(*TracePin, InText);
 		}
-		TracePin->bIsRenaming = false;
+		TracePin->StopRenaming();
 	}
 }
 
@@ -250,7 +272,7 @@ void STraceListRow::Rename(FTraceViewModel& Trace, const FText& InText)
 		return;
 	}
 
-	const FString TraceName = Trace.Name.ToString();
+	const FString TraceName = Trace.GetName();
 	const FString NewTraceName = InText.ToString().TrimStartAndEnd();
 
 	if (NewTraceName == TraceName || NewTraceName.IsEmpty())
@@ -258,7 +280,7 @@ void STraceListRow::Rename(FTraceViewModel& Trace, const FText& InText)
 		return;
 	}
 
-	if (Trace.bIsLive)
+	if (Trace.IsLive())
 	{
 		FText Message = LOCTEXT("RenameLive", "Cannot rename a live session!");
 		TraceStoreWindow->ShowFailMessage(Message);
@@ -267,7 +289,7 @@ void STraceListRow::Rename(FTraceViewModel& Trace, const FText& InText)
 
 	UE_LOG(LogInsightsFrontend, Log, TEXT("[TraceStore] Renaming \"%s\" to \"%s\"..."), *TraceName, *NewTraceName);
 
-	const FString TraceFile = Trace.Uri.ToString();
+	const FString TraceFile = Trace.GetUri();
 	const FString NewTraceFile = FPaths::Combine(FPaths::GetPath(TraceFile), NewTraceName + TEXT(".utrace"));
 
 	FText Reason;
@@ -283,7 +305,7 @@ void STraceListRow::Rename(FTraceViewModel& Trace, const FText& InText)
 	}
 	if (bHasReservedCharacters || !FPaths::ValidatePath(NewTraceFile, &Reason))
 	{
-		FText Message = FText::Format(LOCTEXT("RenameFailFmt3", "Failed to rename \"{0}\" to \"{1}\"!\n{2}"), Trace.Name, FText::FromString(NewTraceName), Reason);
+		FText Message = FText::Format(LOCTEXT("RenameFailFmt3", "Failed to rename \"{0}\" to \"{1}\"!\n{2}"), Trace.GetNameText(), FText::FromString(NewTraceName), Reason);
 		TraceStoreWindow->ShowFailMessage(Message);
 		return;
 	}
@@ -292,7 +314,7 @@ void STraceListRow::Rename(FTraceViewModel& Trace, const FText& InText)
 	{
 		UE_LOG(LogInsightsFrontend, Warning, TEXT("[TraceStore] Failed to rename \"%s\" to \"%s\"! File already exists."), *TraceName, *NewTraceName);
 
-		FText Message = FText::Format(LOCTEXT("RenameFailFmt1", "Failed to rename \"{0}\" to \"{1}\"!\nFile already exists."), Trace.Name, FText::FromString(NewTraceName));
+		FText Message = FText::Format(LOCTEXT("RenameFailFmt1", "Failed to rename \"{0}\" to \"{1}\"!\nFile already exists."), Trace.GetNameText(), FText::FromString(NewTraceName));
 		TraceStoreWindow->ShowFailMessage(Message);
 		return;
 	}
@@ -302,18 +324,18 @@ void STraceListRow::Rename(FTraceViewModel& Trace, const FText& InText)
 	{
 		UE_LOG(LogInsightsFrontend, Warning, TEXT("[TraceStore] Failed to rename \"%s\" to \"%s\"!"), *TraceName, *NewTraceName);
 
-		FText Message = FText::Format(LOCTEXT("RenameFailFmt2", "Failed to rename \"{0}\" to \"{1}\"!"), Trace.Name, FText::FromString(NewTraceName));
+		FText Message = FText::Format(LOCTEXT("RenameFailFmt2", "Failed to rename \"{0}\" to \"{1}\"!"), Trace.GetNameText(), FText::FromString(NewTraceName));
 		TraceStoreWindow->ShowFailMessage(Message);
 		return;
 	}
 
 	UE_LOG(LogInsightsFrontend, Verbose, TEXT("[TraceStore] Renamed utrace file (\"%s\")."), *NewTraceFile);
-	Trace.Name = FText::FromString(NewTraceName);
-	Trace.Uri = FText::FromString(NewTraceFile);
+	Trace.SetName(NewTraceName);
+	Trace.SetUri(NewTraceFile);
 
-	TraceStoreWindow->TraceViewModelMap.Remove(Trace.TraceId);
-	Trace.TraceId = FTraceViewModel::InvalidTraceId; // cannot be open until its TraceId is updated
-	Trace.ChangeSerial = 0; // to force update
+	TraceStoreWindow->TraceViewModelMap.Remove(Trace.GetId());
+	Trace.SetId(FTraceViewModel::InvalidTraceId); // cannot be open until its TraceId is updated
+	Trace.SetChangeSerial(0); // to force update
 
 	const FString CacheFile = FPaths::ChangeExtension(TraceFile, TEXT("ucache"));
 	if (FPaths::FileExists(CacheFile))
@@ -338,7 +360,7 @@ FText STraceListRow::GetTraceIndexAndId() const
 	TSharedPtr<FTraceViewModel> TracePin = WeakTrace.Pin();
 	if (TracePin.IsValid())
 	{
-		const FString TraceIdStr = FString::Printf(TEXT("0x%X"), TracePin->TraceId);
+		const FString TraceIdStr = FString::Printf(TEXT("0x%X"), TracePin->GetId());
 		return FText::FromString(TraceIdStr);
 	}
 	else
@@ -354,7 +376,7 @@ FText STraceListRow::GetTraceName() const
 	TSharedPtr<FTraceViewModel> TracePin = WeakTrace.Pin();
 	if (TracePin.IsValid())
 	{
-		return TracePin->Name;
+		return TracePin->GetNameText();
 	}
 	else
 	{
@@ -369,7 +391,7 @@ FSlateColor STraceListRow::GetTraceTextColor() const
 	TSharedPtr<FTraceViewModel> TracePin = WeakTrace.Pin();
 	if (TracePin.IsValid())
 	{
-		if (TracePin->TraceId == FTraceViewModel::InvalidTraceId)
+		if (!TracePin->HasValidId())
 		{
 			return FSlateColor(EStyleColor::White25);
 		}
@@ -411,7 +433,7 @@ FText STraceListRow::GetTraceUri() const
 	TSharedPtr<FTraceViewModel> TracePin = WeakTrace.Pin();
 	if (TracePin.IsValid())
 	{
-		return TracePin->Uri;
+		return TracePin->GetUriText();
 	}
 	else
 	{
@@ -426,7 +448,7 @@ FSlateColor STraceListRow::GetColorForPath() const
 	TSharedPtr<FTraceViewModel> TracePin = WeakTrace.Pin();
 	if (TracePin.IsValid())
 	{
-		return TracePin->DirectoryColor;
+		return TracePin->GetDirectoryColor();
 	}
 	else
 	{
@@ -441,7 +463,7 @@ FText STraceListRow::GetTracePlatform() const
 	TSharedPtr<FTraceViewModel> TracePin = WeakTrace.Pin();
 	if (TracePin.IsValid())
 	{
-		return TracePin->Platform;
+		return TracePin->GetPlatformText();
 	}
 	else
 	{
@@ -456,7 +478,7 @@ FText STraceListRow::GetTraceAppName() const
 	TSharedPtr<FTraceViewModel> TracePin = WeakTrace.Pin();
 	if (TracePin.IsValid())
 	{
-		return TracePin->AppName;
+		return TracePin->GetAppNameText();
 	}
 	else
 	{
@@ -466,39 +488,29 @@ FText STraceListRow::GetTraceAppName() const
 
 ////////////////////////////////////////////////////////////////////////////////////////////////////
 
-FText STraceListRow::GetTraceCommandLine() const
+FText STraceListRow::GetTraceBuildConfiguration() const
 {
 	TSharedPtr<FTraceViewModel> TracePin = WeakTrace.Pin();
 	if (TracePin.IsValid())
 	{
-		return TracePin->CommandLine;
-	}
-	else
-	{
-		return FText::GetEmpty();
+		if (TracePin->GetBuildConfiguration() != EBuildConfiguration::Unknown)
+		{
+			return EBuildConfigurations::ToText(TracePin->GetBuildConfiguration());
+		}
 	}
+	return FText::GetEmpty();
 }
 
 ////////////////////////////////////////////////////////////////////////////////////////////////////
 
-FText STraceListRow::GetTraceCommandLineHighlightText() const
+FText STraceListRow::GetTraceBuildTarget() const
 {
-	TSharedPtr<STraceStoreWindow> ParentWidgetPin = WeakParentWidget.Pin();
-	if (ParentWidgetPin.IsValid())
+	TSharedPtr<FTraceViewModel> TracePin = WeakTrace.Pin();
+	if (TracePin.IsValid())
 	{
-		if (ParentWidgetPin->bSearchByCommandLine && ParentWidgetPin->FilterByNameSearchBox.IsValid())
+		if (TracePin->GetBuildTarget() != EBuildTargetType::Unknown)
 		{
-			const FText SearchText = ParentWidgetPin->FilterByNameSearchBox->GetText();
-			bool bQuotesRemoved = false;
-			const FString SearchString = SearchText.ToString().TrimQuotes(&bQuotesRemoved);
-			if (bQuotesRemoved)
-			{
-				return FText::FromString(SearchString);
-			}
-			else
-			{
-				return SearchText;
-			}
+			return FText::FromString(LexToString(TracePin->GetBuildTarget()));
 		}
 	}
 	return FText::GetEmpty();
@@ -506,12 +518,12 @@ FText STraceListRow::GetTraceCommandLineHighlightText() const
 
 ////////////////////////////////////////////////////////////////////////////////////////////////////
 
-FText STraceListRow::GetTraceBranch() const
+FText STraceListRow::GetTraceBuildVersion() const
 {
 	TSharedPtr<FTraceViewModel> TracePin = WeakTrace.Pin();
 	if (TracePin.IsValid())
 	{
-		return TracePin->Branch;
+		return TracePin->GetBuildVersionText();
 	}
 	else
 	{
@@ -521,12 +533,12 @@ FText STraceListRow::GetTraceBranch() const
 
 ////////////////////////////////////////////////////////////////////////////////////////////////////
 
-FText STraceListRow::GetTraceBuildVersion() const
+FText STraceListRow::GetTraceBranch() const
 {
 	TSharedPtr<FTraceViewModel> TracePin = WeakTrace.Pin();
 	if (TracePin.IsValid())
 	{
-		return TracePin->BuildVersion;
+		return TracePin->GetBuildBranchText();
 	}
 	else
 	{
@@ -541,7 +553,7 @@ FText STraceListRow::GetTraceChangelist() const
 	TSharedPtr<FTraceViewModel> TracePin = WeakTrace.Pin();
 	if (TracePin.IsValid())
 	{
-		return FText::AsNumber(TracePin->Changelist, &FNumberFormattingOptions::DefaultNoGrouping());
+		return FText::AsNumber(TracePin->GetChangelist(), &FNumberFormattingOptions::DefaultNoGrouping());
 	}
 	else
 	{
@@ -556,7 +568,7 @@ EVisibility STraceListRow::TraceChangelistVisibility() const
 	TSharedPtr<FTraceViewModel> TracePin = WeakTrace.Pin();
 	if (TracePin.IsValid())
 	{
-		return TracePin->Changelist != 0 ? EVisibility::Visible : EVisibility::Collapsed;
+		return TracePin->GetChangelist() != 0 ? EVisibility::Visible : EVisibility::Collapsed;
 	}
 	else
 	{
@@ -566,29 +578,54 @@ EVisibility STraceListRow::TraceChangelistVisibility() const
 
 ////////////////////////////////////////////////////////////////////////////////////////////////////
 
-FText STraceListRow::GetTraceBuildConfiguration() const
+FText STraceListRow::GetTraceEngineVersion() const
 {
 	TSharedPtr<FTraceViewModel> TracePin = WeakTrace.Pin();
 	if (TracePin.IsValid())
 	{
-		if (TracePin->ConfigurationType != EBuildConfiguration::Unknown)
-		{
-			return EBuildConfigurations::ToText(TracePin->ConfigurationType);
-		}
+		return TracePin->GetEngineVersionText();
+	}
+	else
+	{
+		return FText::GetEmpty();
 	}
-	return FText::GetEmpty();
 }
 
 ////////////////////////////////////////////////////////////////////////////////////////////////////
 
-FText STraceListRow::GetTraceBuildTarget() const
+FText STraceListRow::GetTraceCommandLine() const
 {
 	TSharedPtr<FTraceViewModel> TracePin = WeakTrace.Pin();
 	if (TracePin.IsValid())
 	{
-		if (TracePin->TargetType != EBuildTargetType::Unknown)
+		return TracePin->GetCommandLineText();
+	}
+	else
+	{
+		return FText::GetEmpty();
+	}
+}
+
+////////////////////////////////////////////////////////////////////////////////////////////////////
+
+FText STraceListRow::GetTraceCommandLineHighlightText() const
+{
+	TSharedPtr<STraceStoreWindow> ParentWidgetPin = WeakParentWidget.Pin();
+	if (ParentWidgetPin.IsValid())
+	{
+		if (ParentWidgetPin->bSearchByCommandLine && ParentWidgetPin->FilterByNameSearchBox.IsValid())
 		{
-			return FText::FromString(LexToString(TracePin->TargetType));
+			const FText SearchText = ParentWidgetPin->FilterByNameSearchBox->GetText();
+			bool bQuotesRemoved = false;
+			const FString SearchString = SearchText.ToString().TrimQuotes(&bQuotesRemoved);
+			if (bQuotesRemoved)
+			{
+				return FText::FromString(SearchString);
+			}
+			else
+			{
+				return SearchText;
+			}
 		}
 	}
 	return FText::GetEmpty();
@@ -596,12 +633,27 @@ FText STraceListRow::GetTraceBuildTarget() const
 
 ////////////////////////////////////////////////////////////////////////////////////////////////////
 
+FText STraceListRow::GetTraceVFSPaths() const
+{
+	TSharedPtr<FTraceViewModel> TracePin = WeakTrace.Pin();
+	if (TracePin.IsValid())
+	{
+		return TracePin->GetVFSPathsText();
+	}
+	else
+	{
+		return FText::GetEmpty();
+	}
+}
+
+////////////////////////////////////////////////////////////////////////////////////////////////////
+
 FText STraceListRow::GetTraceTimestamp() const
 {
 	TSharedPtr<FTraceViewModel> TracePin = WeakTrace.Pin();
 	if (TracePin.IsValid())
 	{
-		return FText::AsDate(TracePin->Timestamp);
+		return FText::AsDate(TracePin->GetTimestamp());
 	}
 	else
 	{
@@ -616,7 +668,7 @@ FText STraceListRow::GetTraceTimestampForTooltip() const
 	TSharedPtr<FTraceViewModel> TracePin = WeakTrace.Pin();
 	if (TracePin.IsValid())
 	{
-		return FText::AsDateTime(TracePin->Timestamp);
+		return FText::AsDateTime(TracePin->GetTimestamp());
 	}
 	else
 	{
@@ -635,7 +687,7 @@ FText STraceListRow::GetTraceSize() const
 		//FormattingOptions.MinimumFractionalDigits = 1;
 		//FormattingOptions.MaximumFractionalDigits = 1;
 		//return FText::AsMemory(TracePin->Size, &FormattingOptions);
-		return FText::Format(LOCTEXT("SessionFileSizeFormatKiB", "{0} KiB"), TracePin->Size / 1024);
+		return FText::Format(LOCTEXT("SessionFileSizeFormatKiB", "{0} KiB"), TracePin->GetSize() / 1024);
 	}
 	else
 	{
@@ -650,13 +702,13 @@ FText STraceListRow::GetTraceSizeForTooltip() const
 	TSharedPtr<FTraceViewModel> TracePin = WeakTrace.Pin();
 	if (TracePin.IsValid())
 	{
-		if (TracePin->Size > 1024)
+		if (TracePin->GetSize() > 1024)
 		{
-			return FText::Format(LOCTEXT("TraceTooltip_FileSize2", "{0} bytes ({1})"), FText::AsNumber(TracePin->Size), FText::AsMemory(TracePin->Size));
+			return FText::Format(LOCTEXT("TraceTooltip_FileSize2", "{0} bytes ({1})"), FText::AsNumber(TracePin->GetSize()), FText::AsMemory(TracePin->GetSize()));
 		}
 		else
 		{
-			return FText::Format(LOCTEXT("TraceTooltip_FileSize1", "{0} bytes"), FText::AsNumber(TracePin->Size));
+			return FText::Format(LOCTEXT("TraceTooltip_FileSize1", "{0} bytes"), FText::AsNumber(TracePin->GetSize()));
 		}
 	}
 	else
@@ -672,7 +724,7 @@ FSlateColor STraceListRow::GetColorBySize() const
 	TSharedPtr<FTraceViewModel> TracePin = WeakTrace.Pin();
 	if (TracePin.IsValid())
 	{
-		if (TracePin->Size < 1024ULL * 1024ULL)
+		if (TracePin->GetSize() < 1024ULL * 1024ULL)
 		{
 			// < 1 MiB
 			TSharedRef<ITypedTableView<TSharedPtr<FTraceViewModel>>> OwnerWidget = OwnerTablePtr.Pin().ToSharedRef();
@@ -687,7 +739,7 @@ FSlateColor STraceListRow::GetColorBySize() const
 				return FSlateColor(FLinearColor(0.5f, 0.5f, 0.5f, 1.0f));
 			}
 		}
-		else if (TracePin->Size < 1024ULL * 1024ULL * 1024ULL)
+		else if (TracePin->GetSize() < 1024ULL * 1024ULL * 1024ULL)
 		{
 			// [1 MiB  .. 1 GiB)
 			return FSlateColor(FLinearColor(1.0f, 1.0f, 1.0f, 1.0f));
@@ -711,7 +763,7 @@ FText STraceListRow::GetTraceStatus() const
 	TSharedPtr<FTraceViewModel> TracePin = WeakTrace.Pin();
 	if (TracePin.IsValid())
 	{
-		if (TracePin->bIsLive)
+		if (TracePin->IsLive())
 		{
 			return LOCTEXT("LiveTraceStatus", "LIVE");
 		}
@@ -726,13 +778,14 @@ FText STraceListRow::GetTraceStatusForTooltip() const
 	TSharedPtr<FTraceViewModel> TracePin = WeakTrace.Pin();
 	if (TracePin.IsValid())
 	{
-		if (TracePin->bIsLive)
+		if (TracePin->IsLive())
 		{
+			const uint32 IpAddress = TracePin->GetIpAddress();
 			FString Ip = FString::Printf(TEXT("%d.%d.%d.%d"),
-				(TracePin->IpAddress >> 24) & 0xFF,
-				(TracePin->IpAddress >> 16) & 0xFF,
-				(TracePin->IpAddress >> 8) & 0xFF,
-				TracePin->IpAddress & 0xFF);
+				(IpAddress >> 24) & 0xFF,
+				(IpAddress >> 16) & 0xFF,
+				(IpAddress >> 8) & 0xFF,
+				IpAddress & 0xFF);
 			return FText::Format(LOCTEXT("LiveTraceStatusFmt", "LIVE ({0})"), FText::FromString(Ip));
 		}
 		else
@@ -831,14 +884,16 @@ TSharedPtr<SToolTip> STraceListRow::CreateTooltip() const
 			];
 
 		int32 Row = 0;
+		AddGridPanelRow(GridPanel, Row++, LOCTEXT("TraceTooltip_CommandLine", "Command Line:"), &STraceListRow::GetTraceCommandLine, &STraceListRow::GetTraceCommandLineHighlightText);
 		AddGridPanelRow(GridPanel, Row++, LOCTEXT("TraceTooltip_Platform", "Platform:"), &STraceListRow::GetTracePlatform);
 		AddGridPanelRow(GridPanel, Row++, LOCTEXT("TraceTooltip_AppName", "App Name:"), &STraceListRow::GetTraceAppName);
-		AddGridPanelRow(GridPanel, Row++, LOCTEXT("TraceTooltip_CommandLine", "Command Line:"), &STraceListRow::GetTraceCommandLine, &STraceListRow::GetTraceCommandLineHighlightText);
 		AddGridPanelRow(GridPanel, Row++, LOCTEXT("TraceTooltip_BuildConfig", "Build Config:"), &STraceListRow::GetTraceBuildConfiguration);
 		AddGridPanelRow(GridPanel, Row++, LOCTEXT("TraceTooltip_BuildTarget", "Build Target:"), &STraceListRow::GetTraceBuildTarget);
-		AddGridPanelRow(GridPanel, Row++, LOCTEXT("TraceTooltip_Branch", "Build Branch:"), &STraceListRow::GetTraceBranch);
 		AddGridPanelRow(GridPanel, Row++, LOCTEXT("TraceTooltip_BuildVersion", "Build Version:"), &STraceListRow::GetTraceBuildVersion);
+		AddGridPanelRow(GridPanel, Row++, LOCTEXT("TraceTooltip_Branch", "Build Branch:"), &STraceListRow::GetTraceBranch);
 		AddGridPanelRow(GridPanel, Row++, LOCTEXT("TraceTooltip_Changelist", "Changelist:"), &STraceListRow::GetTraceChangelist, nullptr, &STraceListRow::TraceChangelistVisibility);
+		AddGridPanelRow(GridPanel, Row++, LOCTEXT("TraceTooltip_EngineVersion", "Engine Version:"), &STraceListRow::GetTraceEngineVersion);
+		AddGridPanelRow(GridPanel, Row++, LOCTEXT("TraceTooltip_VFSPaths", "VFS Paths:"), &STraceListRow::GetTraceVFSPaths);
 		AddGridPanelRow(GridPanel, Row++, LOCTEXT("TraceTooltip_Timestamp", "Timestamp:"), &STraceListRow::GetTraceTimestampForTooltip);
 		AddGridPanelRow(GridPanel, Row++, LOCTEXT("TraceTooltip_Size", "File Size:"), &STraceListRow::GetTraceSizeForTooltip);
 		AddGridPanelRow(GridPanel, Row++, LOCTEXT("TraceTooltip_Status", "Status:"), &STraceListRow::GetTraceStatusForTooltip);
diff --git a/Engine/Source/Developer/TraceInsightsFrontend/Private/InsightsFrontend/Widgets/STraceListRow.h b/Engine/Source/Developer/TraceInsightsFrontend/Private/InsightsFrontend/Widgets/STraceListRow.h
--- a/Engine/Source/Developer/TraceInsightsFrontend/Private/InsightsFrontend/Widgets/STraceListRow.h
+++ b/Engine/Source/Developer/TraceInsightsFrontend/Private/InsightsFrontend/Widgets/STraceListRow.h
@@ -16,7 +16,7 @@ class SGridPanel;
 namespace UE::Insights
 {
 
-struct FTraceViewModel;
+class FTraceViewModel;
 class STraceStoreWindow;
 
 ////////////////////////////////////////////////////////////////////////////////////////////////////
@@ -30,8 +30,10 @@ struct FTraceListColumns
 	static const FName AppName;
 	static const FName BuildConfig;
 	static const FName BuildTarget;
-	static const FName BuildBranch;
 	static const FName BuildVersion;
+	static const FName BuildBranch;
+	static const FName Changelist;
+	static const FName EngineVersion;
 	static const FName Size;
 	static const FName Status;
 };
@@ -60,14 +62,16 @@ public:
 	FSlateColor GetColorForPath() const;
 	FText GetTracePlatform() const;
 	FText GetTraceAppName() const;
-	FText GetTraceCommandLine() const;
-	FText GetTraceCommandLineHighlightText() const;
-	FText GetTraceBranch() const;
+	FText GetTraceEngineVersion() const;
 	FText GetTraceBuildVersion() const;
+	FText GetTraceBranch() const;
 	FText GetTraceChangelist() const;
 	EVisibility TraceChangelistVisibility() const;
 	FText GetTraceBuildConfiguration() const;
 	FText GetTraceBuildTarget() const;
+	FText GetTraceCommandLine() const;
+	FText GetTraceCommandLineHighlightText() const;
+	FText GetTraceVFSPaths() const;
 	FText GetTraceTimestamp() const;
 	FText GetTraceTimestampForTooltip() const;
 	FText GetTraceSize() const;
diff --git a/Engine/Source/Developer/TraceInsightsFrontend/Private/InsightsFrontend/Widgets/STraceStoreWindow.cpp b/Engine/Source/Developer/TraceInsightsFrontend/Private/InsightsFrontend/Widgets/STraceStoreWindow.cpp
--- a/Engine/Source/Developer/TraceInsightsFrontend/Private/InsightsFrontend/Widgets/STraceStoreWindow.cpp
+++ b/Engine/Source/Developer/TraceInsightsFrontend/Private/InsightsFrontend/Widgets/STraceStoreWindow.cpp
@@ -332,6 +332,15 @@ TSharedRef<SWidget> STraceStoreWindow::ConstructFiltersToolbar()
 			FSlateIcon(FInsightsFrontendStyle::GetStyleSetName(), "Icons.Filter.ToolBar"),
 			false);
 
+		// Filter by Build Version
+		ToolbarBuilder.AddComboButton(
+			FUIAction(),
+			FOnGetContent::CreateSP(this, &STraceStoreWindow::MakeBuildVersionFilterMenu),
+			LOCTEXT("FilterByBuildVersionText", "Build Version"),
+			LOCTEXT("FilterByBuildVersionToolTip", "Filters the list of trace sessions by Build Version."),
+			FSlateIcon(FInsightsFrontendStyle::GetStyleSetName(), "Icons.Filter.ToolBar"),
+			false);
+
 		// Filter by Branch
 		ToolbarBuilder.AddComboButton(
 			FUIAction(),
@@ -341,12 +350,23 @@ TSharedRef<SWidget> STraceStoreWindow::ConstructFiltersToolbar()
 			FSlateIcon(FInsightsFrontendStyle::GetStyleSetName(), "Icons.Filter.ToolBar"),
 			false);
 
-		// Filter by Version
+#if 0
+		// Filter by Changelist
+		ToolbarBuilder.AddComboButton(
+			FUIAction(),
+			FOnGetContent::CreateSP(this, &STraceStoreWindow::MakeChangelistFilterMenu),
+			LOCTEXT("FilterByChangelistText", "Changelist"),
+			LOCTEXT("FilterByChangelistToolTip", "Filters the list of trace sessions by changelist."),
+			FSlateIcon(FInsightsFrontendStyle::GetStyleSetName(), "Icons.Filter.ToolBar"),
+			false);
+#endif
+
+		// Filter by Engine Version
 		ToolbarBuilder.AddComboButton(
 			FUIAction(),
-			FOnGetContent::CreateSP(this, &STraceStoreWindow::MakeVersionFilterMenu),
-			LOCTEXT("FilterByVersionText", "Version"),
-			LOCTEXT("FilterByVersionToolTip", "Filters the list of trace sessions by Version."),
+			FOnGetContent::CreateSP(this, &STraceStoreWindow::MakeEngineVersionFilterMenu),
+			LOCTEXT("FilterByEngineVersionText", "Engine Version"),
+			LOCTEXT("FilterByEngineVersionToolTip", "Filters the list of trace sessions by Engine Version."),
 			FSlateIcon(FInsightsFrontendStyle::GetStyleSetName(), "Icons.Filter.ToolBar"),
 			false);
 
@@ -512,6 +532,22 @@ TSharedRef<SWidget> STraceStoreWindow::ConstructSessionsPanel()
 				]
 			]
 
+			+ SHeaderRow::Column(FTraceListColumns::BuildVersion)
+			.FillWidth(0.25f)
+			.InitialSortMode(EColumnSortMode::Ascending)
+			.SortMode(this, &STraceStoreWindow::GetSortModeForColumn, FTraceListColumns::BuildVersion)
+			.OnSort(this, &STraceStoreWindow::OnSortModeChanged)
+			.OnGetMenuContent(this, &STraceStoreWindow::MakeBuildVersionColumnHeaderMenu)
+			[
+				SNew(SBox)
+				.VAlign(VAlign_Center)
+				[
+					SNew(STextBlock)
+					.Text(LOCTEXT("BuildVersionColumn", "Build Version"))
+					.ColorAndOpacity_Lambda([this] { return FilterByBuildVersion->IsEmpty() ? FLinearColor(0.5f, 0.5f, 0.5f, 1.0f) : FLinearColor(0.3f, 0.75f, 1.0f, 1.0f); })
+				]
+			]
+
 			+ SHeaderRow::Column(FTraceListColumns::BuildBranch)
 			.FillWidth(0.2f)
 			.InitialSortMode(EColumnSortMode::Ascending)
@@ -528,19 +564,37 @@ TSharedRef<SWidget> STraceStoreWindow::ConstructSessionsPanel()
 				]
 			]
 
-			+ SHeaderRow::Column(FTraceListColumns::BuildVersion)
-			.FillWidth(0.25f)
+#if 0
+			+ SHeaderRow::Column(FTraceListColumns::Changelist)
+			.FillWidth(0.1f)
 			.InitialSortMode(EColumnSortMode::Ascending)
-			.SortMode(this, &STraceStoreWindow::GetSortModeForColumn, FTraceListColumns::BuildVersion)
+			.SortMode(this, &STraceStoreWindow::GetSortModeForColumn, FTraceListColumns::Changelist)
 			.OnSort(this, &STraceStoreWindow::OnSortModeChanged)
-			.OnGetMenuContent(this, &STraceStoreWindow::MakeVersionColumnHeaderMenu)
+			.OnGetMenuContent(this, &STraceStoreWindow::MakeChangelistColumnHeaderMenu)
 			[
 				SNew(SBox)
 				.VAlign(VAlign_Center)
 				[
 					SNew(STextBlock)
-					.Text(LOCTEXT("BuildVersionColumn", "Build Version"))
-					.ColorAndOpacity_Lambda([this] { return FilterByVersion->IsEmpty() ? FLinearColor(0.5f, 0.5f, 0.5f, 1.0f) : FLinearColor(0.3f, 0.75f, 1.0f, 1.0f); })
+					.Text(LOCTEXT("ChangelistColumn", "CL"))
+					.ColorAndOpacity_Lambda([this] { return FilterByChangelist->IsEmpty() ? FLinearColor(0.5f, 0.5f, 0.5f, 1.0f) : FLinearColor(0.3f, 0.75f, 1.0f, 1.0f); })
+				]
+			]
+#endif
+
+			+ SHeaderRow::Column(FTraceListColumns::EngineVersion)
+			.FillWidth(0.1f)
+			.InitialSortMode(EColumnSortMode::Ascending)
+			.SortMode(this, &STraceStoreWindow::GetSortModeForColumn, FTraceListColumns::EngineVersion)
+			.OnSort(this, &STraceStoreWindow::OnSortModeChanged)
+			.OnGetMenuContent(this, &STraceStoreWindow::MakeEngineVersionColumnHeaderMenu)
+			[
+				SNew(SBox)
+				.VAlign(VAlign_Center)
+				[
+					SNew(STextBlock)
+					.Text(LOCTEXT("EngineVersionColumn", "UE"))
+					.ColorAndOpacity_Lambda([this] { return FilterByEngineVersion->IsEmpty() ? FLinearColor(0.5f, 0.5f, 0.5f, 1.0f) : FLinearColor(0.3f, 0.75f, 1.0f, 1.0f); })
 				]
 			]
 
@@ -1123,8 +1177,8 @@ bool STraceStoreWindow::CanRenameSelectedTrace() const
 	}
 	TSharedPtr<FTraceViewModel> SelectedTrace = GetSingleSelectedTrace();
 	return SelectedTrace.IsValid()
-		&& SelectedTrace->TraceId != FTraceViewModel::InvalidTraceId
-		&& !SelectedTrace->bIsLive;
+		&& SelectedTrace->HasValidId()
+		&& !SelectedTrace->IsLive();
 }
 
 ////////////////////////////////////////////////////////////////////////////////////////////////////
@@ -1139,9 +1193,9 @@ void STraceStoreWindow::RenameSelectedTrace()
 	}
 
 	TSharedPtr<FTraceViewModel> SelectedTrace = GetSingleSelectedTrace();
-	SelectedTrace->bIsRenaming = true;
+	SelectedTrace->StartRenaming();
 
-	TSharedPtr<SEditableTextBox> RenameTextBox = SelectedTrace->RenameTextBox.Pin();
+	TSharedPtr<SEditableTextBox> RenameTextBox = SelectedTrace->GetRenameTextBox();
 	if (RenameTextBox.IsValid())
 	{
 		FSlateApplication::Get().SetKeyboardFocus(RenameTextBox.ToSharedRef(), EFocusCause::SetDirectly);
@@ -1161,8 +1215,7 @@ bool STraceStoreWindow::CanDeleteSelectedTraces() const
 	TArray<TSharedPtr<FTraceViewModel>> SelectedTraces = TraceListView->GetSelectedItems();
 	for (const TSharedPtr<FTraceViewModel>& SelectedTrace : SelectedTraces)
 	{
-		if ((SelectedTrace->TraceId != FTraceViewModel::InvalidTraceId) &&
-			!SelectedTrace->bIsLive)
+		if (SelectedTrace->HasValidId() && !SelectedTrace->IsLive())
 		{
 			return true;
 		}
@@ -1183,7 +1236,7 @@ void STraceStoreWindow::DeleteSelectedTraces()
 
 	TArray<TSharedPtr<FTraceViewModel>> TracesToDelete = TraceListView->GetSelectedItems();
 	// Filter the traces that can actually be deleted :
-	TracesToDelete.RemoveAll([](const TSharedPtr<FTraceViewModel>& InTrace) { return (InTrace->TraceId == FTraceViewModel::InvalidTraceId) || InTrace->bIsLive; });
+	TracesToDelete.RemoveAll([](const TSharedPtr<FTraceViewModel>& InTrace) { return !InTrace->HasValidId() || InTrace->IsLive(); });
 	if (TracesToDelete.Num() == 0)
 	{
 		return;
@@ -1196,7 +1249,7 @@ void STraceStoreWindow::DeleteSelectedTraces()
 		TStringBuilder<2048> TraceFilesToDelete;
 		for (int32 TraceIndex = 0; TraceIndex < TracesToDelete.Num() && TraceIndex < 3; ++TraceIndex)
 		{
-			TraceFilesToDelete.Append(TracesToDelete[TraceIndex]->Uri.ToString());
+			TraceFilesToDelete.Append(TracesToDelete[TraceIndex]->GetUri());
 			TraceFilesToDelete.Append(TEXT("\n"));
 		}
 		if (TracesToDelete.Num() > 3)
@@ -1254,7 +1307,7 @@ void STraceStoreWindow::DeleteSelectedTraces()
 			FilteredTraceViewModels.Remove(TraceToDelete);
 
 			TraceViewModels.Remove(TraceToDelete);
-			TraceViewModelMap.Remove(TraceToDelete->TraceId);
+			TraceViewModelMap.Remove(TraceToDelete->GetId());
 		}
 	}
 
@@ -1284,17 +1337,17 @@ void STraceStoreWindow::DeleteSelectedTraces()
 
 bool STraceStoreWindow::DeleteTrace(const TSharedPtr<FTraceViewModel>& TraceToDelete)
 {
-	FString TraceName = TraceToDelete->Name.ToString();
+	FString TraceName = TraceToDelete->GetName();
 	UE_LOG(LogInsightsFrontend, Log, TEXT("[TraceStore] Deleting \"%s\"..."), *TraceName);
 
-	if (TraceToDelete->bIsLive)
+	if (TraceToDelete->IsLive())
 	{
 		FText Message = FText::Format(LOCTEXT("CannotDeleteLiveTraceFmt", "Cannot delete a live trace (\"{0}\")!"), FText::FromString(TraceName));
 		ShowFailMessage(Message);
 		return false;
 	}
 
-	FString TraceFile = TraceToDelete->Uri.ToString();
+	FString TraceFile = TraceToDelete->GetUri();
 	if (!FPaths::FileExists(TraceFile) || !IFileManager::Get().Delete(*TraceFile))
 	{
 		UE_LOG(LogInsightsFrontend, Warning, TEXT("[TraceStore] Failed to delete trace file (\"%s\")!"), *TraceFile);
@@ -1323,8 +1376,7 @@ bool STraceStoreWindow::DeleteTrace(const TSharedPtr<FTraceViewModel>& TraceToDe
 bool STraceStoreWindow::CanCopyTraceId() const
 {
 	TSharedPtr<FTraceViewModel> SelectedTrace = GetSingleSelectedTrace();
-	return SelectedTrace.IsValid()
-		&& SelectedTrace->TraceId != FTraceViewModel::InvalidTraceId;
+	return SelectedTrace.IsValid() && SelectedTrace->HasValidId();
 }
 
 ////////////////////////////////////////////////////////////////////////////////////////////////////
@@ -1334,7 +1386,7 @@ void STraceStoreWindow::CopyTraceId()
 	if (CanCopyTraceId())
 	{
 		TSharedPtr<FTraceViewModel> SelectedTrace = GetSingleSelectedTrace();
-		FString ClipboardText = FString::Printf(TEXT("0x%X"), SelectedTrace->TraceId);
+		FString ClipboardText = FString::Printf(TEXT("0x%X"), SelectedTrace->GetId());
 		FPlatformApplicationMisc::ClipboardCopy(*ClipboardText);
 	}
 }
@@ -1344,8 +1396,7 @@ void STraceStoreWindow::CopyTraceId()
 bool STraceStoreWindow::CanCopyFullPath() const
 {
 	TSharedPtr<FTraceViewModel> SelectedTrace = GetSingleSelectedTrace();
-	return SelectedTrace.IsValid()
-		&& SelectedTrace->TraceId != FTraceViewModel::InvalidTraceId;
+	return SelectedTrace.IsValid() && SelectedTrace->HasValidId();
 }
 
 ////////////////////////////////////////////////////////////////////////////////////////////////////
@@ -1355,7 +1406,7 @@ void STraceStoreWindow::CopyFullPath()
 	if (CanCopyFullPath())
 	{
 		TSharedPtr<FTraceViewModel> SelectedTrace = GetSingleSelectedTrace();
-		FPlatformApplicationMisc::ClipboardCopy(*SelectedTrace->Uri.ToString());
+		FPlatformApplicationMisc::ClipboardCopy(*SelectedTrace->GetUri());
 	}
 }
 
@@ -1368,8 +1419,7 @@ bool STraceStoreWindow::CanOpenContainingFolder() const
 		return false;
 	}
 	TSharedPtr<FTraceViewModel> SelectedTrace = GetSingleSelectedTrace();
-	return SelectedTrace.IsValid()
-		&& SelectedTrace->TraceId != FTraceViewModel::InvalidTraceId;
+	return SelectedTrace.IsValid() && SelectedTrace->HasValidId();
 }
 
 ////////////////////////////////////////////////////////////////////////////////////////////////////
@@ -1381,7 +1431,7 @@ void STraceStoreWindow::OpenContainingFolder()
 	if (CanOpenContainingFolder())
 	{
 		TSharedPtr<FTraceViewModel> SelectedTrace = GetSingleSelectedTrace();
-		FPlatformProcess::ExploreFolder(*SelectedTrace->Uri.ToString());
+		FPlatformProcess::ExploreFolder(*SelectedTrace->GetUri());
 	}
 }
 
@@ -1389,7 +1439,7 @@ void STraceStoreWindow::OpenContainingFolder()
 
 bool STraceStoreWindow::HasAnyLiveTrace() const
 {
-	return TraceViewModels.FindByPredicate(&FTraceViewModel::bIsLive) != nullptr;
+	return TraceViewModels.FindByPredicate([](const TSharedPtr<FTraceViewModel>& Trace) { return Trace->IsLive(); }) != nullptr;
 }
 
 ////////////////////////////////////////////////////////////////////////////////////////////////////
@@ -1477,7 +1527,7 @@ FSlateColor STraceStoreWindow::GetColorByPath(const FString& Uri)
 		return FAppStyle::Get().GetSlateColor((*Dir)->Color);
 	}
 	// If this is default trace store directory, use foreground
-	return FSlateColor(FSlateColor::UseForeground());
+	return FSlateColor::UseForeground();
 }
 
 ////////////////////////////////////////////////////////////////////////////////////////////////////
@@ -1579,13 +1629,13 @@ void STraceStoreWindow::RefreshTraceList()
 				for (int32 TraceIndex = 0; TraceIndex < TraceViewModelCount; ++TraceIndex)
 				{
 					FTraceViewModel& Trace = *TraceViewModels[TraceIndex];
-					const TSharedPtr<FStoreBrowserTraceInfo>* InTracePtrPtr = InTraceMap.Find(Trace.TraceId);
+					const TSharedPtr<FStoreBrowserTraceInfo>* InTracePtrPtr = InTraceMap.Find(Trace.GetId());
 					if (!InTracePtrPtr)
 					{
 						// This trace was removed.
 						RemovedTraces++;
 						TraceViewModels.RemoveAtSwap(TraceIndex);
-						TraceViewModelMap.Remove(Trace.TraceId);
+						TraceViewModelMap.Remove(Trace.GetId());
 						TraceIndex--;
 						TraceViewModelCount--;
 					}
@@ -1600,7 +1650,7 @@ void STraceStoreWindow::RefreshTraceList()
 				if (TracePtrPtr)
 				{
 					FTraceViewModel& Trace = **TracePtrPtr;
-					if (Trace.ChangeSerial != SourceTrace.ChangeSerial || bSettingsChanged)
+					if (Trace.GetChangeSerial() != SourceTrace.ChangeSerial || bSettingsChanged)
 					{
 						// This trace was updated or settings updated
 						UpdatedTraces++;
@@ -1612,10 +1662,10 @@ void STraceStoreWindow::RefreshTraceList()
 					// This trace was added.
 					AddedTraces++;
 					TSharedPtr<FTraceViewModel> TracePtr = MakeShared<FTraceViewModel>();
-					TracePtr->TraceId = SourceTrace.TraceId;
+					TracePtr->SetId(SourceTrace.TraceId);
 					UpdateTrace(*TracePtr, SourceTrace);
 					TraceViewModels.Add(TracePtr);
-					TraceViewModelMap.Add(TracePtr->TraceId, TracePtr);
+					TraceViewModelMap.Add(TracePtr->GetId(), TracePtr);
 				}
 			}
 		}
@@ -1649,61 +1699,63 @@ bool STraceStoreWindow::IsConnected() const
 
 void STraceStoreWindow::UpdateTrace(FTraceViewModel& InOutTrace, const FStoreBrowserTraceInfo& InSourceTrace)
 {
-	check(InOutTrace.TraceId == InSourceTrace.TraceId);
+	check(InOutTrace.GetId() == InSourceTrace.TraceId);
 
-	InOutTrace.ChangeSerial = InSourceTrace.ChangeSerial;
+	InOutTrace.SetChangeSerial(InSourceTrace.ChangeSerial);
 
-	InOutTrace.Name = FText::FromString(InSourceTrace.Name);
-	InOutTrace.Uri = FText::FromString(InSourceTrace.Uri);
-	InOutTrace.DirectoryColor = GetColorByPath(InSourceTrace.Uri);
+	InOutTrace.SetName(InSourceTrace.Name);
+	InOutTrace.SetUri(InSourceTrace.Uri);
+	InOutTrace.SetDirectoryColor(GetColorByPath(InSourceTrace.Uri));
 
-	InOutTrace.Timestamp = InSourceTrace.Timestamp;
-	InOutTrace.Size = InSourceTrace.Size;
+	InOutTrace.SetTimestamp(InSourceTrace.Timestamp);
+	InOutTrace.SetSize(InSourceTrace.Size);
 
-	InOutTrace.bIsLive = InSourceTrace.bIsLive;
-	InOutTrace.IpAddress = InSourceTrace.IpAddress;
+	InOutTrace.SetIsLive(InSourceTrace.bIsLive);
+	InOutTrace.SetIpAddress(InSourceTrace.IpAddress);
 
 	// Is metadata updated?
-	if (!InOutTrace.bIsMetadataUpdated && InSourceTrace.MetadataUpdateCount == 0)
+	if (!InOutTrace.IsMetadataUpdated() && InSourceTrace.MetadataUpdateCount == 0)
 	{
-		InOutTrace.bIsMetadataUpdated = true;
-		InOutTrace.Platform = FText::FromString(InSourceTrace.Platform);
+		InOutTrace.SetIsMetadataUpdated(true);
+		InOutTrace.SetPlatform(InSourceTrace.Platform);
 		if (!InSourceTrace.ProjectName.IsEmpty())
 		{
-			InOutTrace.AppName = FText::FromString(InSourceTrace.ProjectName);
+			InOutTrace.SetAppName(InSourceTrace.ProjectName);
 		}
 		else
 		{
-			InOutTrace.AppName = FText::FromString(InSourceTrace.AppName);
+			InOutTrace.SetAppName(InSourceTrace.AppName);
 		}
-		InOutTrace.CommandLine = FText::FromString(InSourceTrace.CommandLine);
-		InOutTrace.Branch = FText::FromString(InSourceTrace.Branch);
-		InOutTrace.BuildVersion = FText::FromString(InSourceTrace.BuildVersion);
-		InOutTrace.Changelist = InSourceTrace.Changelist;
-		InOutTrace.ConfigurationType = InSourceTrace.ConfigurationType;
-		InOutTrace.TargetType = InSourceTrace.TargetType;
+		InOutTrace.SetBuildConfiguration(InSourceTrace.ConfigurationType);
+		InOutTrace.SetBuildTarget(InSourceTrace.TargetType);
+		InOutTrace.SetBuildVersion(InSourceTrace.BuildVersion);
+		InOutTrace.SetBuildBranch(InSourceTrace.Branch);
+		InOutTrace.SetChangelist(InSourceTrace.Changelist);
+		InOutTrace.SetEngineVersion(InSourceTrace.EngineVersion);
+		InOutTrace.SetCommandLine(InSourceTrace.CommandLine);
+		InOutTrace.SetVFSPaths(InSourceTrace.VFSPaths);
 	}
 
 	const FInsightsFrontendSettings& Settings = GetSettings();
 
 	// Auto start analysis for a live trace session.
-	if (InOutTrace.bIsLive &&
-		InOutTrace.bIsMetadataUpdated &&
+	if (InOutTrace.IsLive() &&
+		InOutTrace.IsMetadataUpdated() &&
 		Settings.IsAutoStartAnalysisEnabled() && // is auto start enabled?
-		!AutoStartedSessions.Contains(InOutTrace.TraceId)) // is not already auto-started?
+		!AutoStartedSessions.Contains(InOutTrace.GetId())) // is not already auto-started?
 	{
 		const FString& AutoStartPlatformFilterStr = Settings.GetAutoStartAnalysisPlatform();
 		const FString& AutoStartAppNameFilterStr = Settings.GetAutoStartAnalysisAppName();
 
 		// matches filter?
-		if ((AutoStartPlatformFilterStr.IsEmpty() || FCString::Strcmp(*AutoStartPlatformFilterStr, *InOutTrace.Platform.ToString()) == 0) &&
-			(AutoStartAppNameFilterStr.IsEmpty() || FCString::Strcmp(*AutoStartAppNameFilterStr, *InOutTrace.AppName.ToString()) == 0) &&
-			(AutoStartConfigurationTypeFilter == EBuildConfiguration::Unknown || AutoStartConfigurationTypeFilter == InOutTrace.ConfigurationType) &&
-			(AutoStartTargetTypeFilter == EBuildTargetType::Unknown || AutoStartTargetTypeFilter == InOutTrace.TargetType))
+		if ((AutoStartPlatformFilterStr.IsEmpty() || FCString::Strcmp(*AutoStartPlatformFilterStr, *InOutTrace.GetPlatform()) == 0) &&
+			(AutoStartAppNameFilterStr.IsEmpty() || FCString::Strcmp(*AutoStartAppNameFilterStr, *InOutTrace.GetAppName()) == 0) &&
+			(AutoStartConfigurationTypeFilter == EBuildConfiguration::Unknown || AutoStartConfigurationTypeFilter == InOutTrace.GetBuildConfiguration()) &&
+			(AutoStartTargetTypeFilter == EBuildTargetType::Unknown || AutoStartTargetTypeFilter == InOutTrace.GetBuildTarget()))
 		{
-			UE_LOG(LogInsightsFrontend, Log, TEXT("[TraceStore] Auto starting analysis for trace with id 0x%08X..."), InOutTrace.TraceId);
-			AutoStartedSessions.Add(InOutTrace.TraceId);
-			OpenTraceSession(InOutTrace.TraceId);
+			UE_LOG(LogInsightsFrontend, Log, TEXT("[TraceStore] Auto starting analysis for trace with id 0x%08X..."), InOutTrace.GetId());
+			AutoStartedSessions.Add(InOutTrace.GetId());
+			OpenTraceSession(InOutTrace.GetId());
 		}
 	}
 }
@@ -1733,14 +1785,14 @@ void STraceStoreWindow::UpdateTraceListView()
 		TArray<TSharedPtr<FTraceViewModel>> SelectedTraces = TraceListView->GetSelectedItems();
 		for (const TSharedPtr<FTraceViewModel>& SelectedTrace : SelectedTraces)
 		{
-			TSharedPtr<FTraceViewModel>* FoundNewTrace = TraceViewModelMap.Find(SelectedTrace->TraceId);
+			TSharedPtr<FTraceViewModel>* FoundNewTrace = TraceViewModelMap.Find(SelectedTrace->GetId());
 			if (!FoundNewTrace)
 			{
-				FoundNewTrace = TraceViewModels.FindByPredicate([SelectedTrace](const TSharedPtr<FTraceViewModel>& Trace) { return Trace->Uri.EqualTo(SelectedTrace->Uri); });
+				FoundNewTrace = TraceViewModels.FindByPredicate([SelectedTrace](const TSharedPtr<FTraceViewModel>& Trace) { return Trace->HasSameUri(*SelectedTrace); });
 			}
 			if (!FoundNewTrace)
 			{
-				FoundNewTrace = TraceViewModels.FindByPredicate([SelectedTrace](const TSharedPtr<FTraceViewModel>& Trace) { return Trace->Name.EqualTo(SelectedTrace->Name); });
+				FoundNewTrace = TraceViewModels.FindByPredicate([SelectedTrace](const TSharedPtr<FTraceViewModel>& Trace) { return Trace->HasSameName(*SelectedTrace); });
 			}
 			if (FoundNewTrace)
 			{
@@ -2076,8 +2128,7 @@ FReply STraceStoreWindow::OnDrop(const FGeometry& MyGeometry, const FDragDropEve
 bool STraceStoreWindow::Open_IsEnabled() const
 {
 	TSharedPtr<FTraceViewModel> SelectedTrace = GetSingleSelectedTrace();
-	return SelectedTrace.IsValid()
-		&& SelectedTrace->TraceId != FTraceViewModel::InvalidTraceId;
+	return SelectedTrace.IsValid() && SelectedTrace->HasValidId();
 }
 
 ////////////////////////////////////////////////////////////////////////////////////////////////////
@@ -2161,10 +2212,9 @@ void STraceStoreWindow::OpenTraceFile(const FString& InTraceFile)
 
 void STraceStoreWindow::OpenTraceSession(TSharedPtr<FTraceViewModel> InTraceSession)
 {
-	if (InTraceSession.IsValid() &&
-		InTraceSession->TraceId != FTraceViewModel::InvalidTraceId)
+	if (InTraceSession.IsValid() && InTraceSession->HasValidId())
 	{
-		OpenTraceSession(InTraceSession->TraceId);
+		OpenTraceSession(InTraceSession->GetId());
 	}
 }
 
@@ -2199,7 +2249,7 @@ void STraceStoreWindow::OpenTraceSession(uint32 InTraceId)
 	if (TraceSessionPtrPtr)
 	{
 		FTraceViewModel& TraceSession = **TraceSessionPtrPtr;
-		SplashScreenOverlayTraceFile = FPaths::GetBaseFilename(TraceSession.Uri.ToString());
+		SplashScreenOverlayTraceFile = FPaths::GetBaseFilename(TraceSession.GetUri());
 	}
 	ShowSplashScreenOverlay();
 }
@@ -2247,7 +2297,7 @@ TSharedRef<SWidget> STraceStoreWindow::MakeTraceListMenu()
 	{
 		// Make a copy of the trace list (to allow list view to be sorted by other criteria).
 		TArray<TSharedPtr<FTraceViewModel>> SortedTraces(TraceViewModels);
-		Algo::SortBy(SortedTraces, &FTraceViewModel::Timestamp);
+		Algo::SortBy(SortedTraces, [](const TSharedPtr<FTraceViewModel>& Trace) { return Trace->GetTimestamp(); });
 
 		int32 TraceCountLimit = 10; // top 10
 
@@ -2256,8 +2306,8 @@ TSharedRef<SWidget> STraceStoreWindow::MakeTraceListMenu()
 		{
 			const FTraceViewModel& Trace = *SortedTraces[TraceIndex];
 
-			FText Label = Trace.Name;
-			if (Trace.bIsLive)
+			FText Label = Trace.GetNameText();
+			if (Trace.IsLive())
 			{
 				Label = FText::Format(LOCTEXT("LiveTraceTextFmt", "{0} (LIVE!)"), Label);
 			}
@@ -2266,7 +2316,7 @@ TSharedRef<SWidget> STraceStoreWindow::MakeTraceListMenu()
 				Label,
 				TAttribute<FText>(), // no tooltip
 				FSlateIcon(),
-				FUIAction(FExecuteAction::CreateSP(this, &STraceStoreWindow::OpenTraceSession, Trace.TraceId)),
+				FUIAction(FExecuteAction::CreateSP(this, &STraceStoreWindow::OpenTraceSession, Trace.GetId())),
 				NAME_None,
 				EUserInterfaceActionType::Button);
 		}
@@ -2508,6 +2558,43 @@ void STraceStoreWindow::BuildBuildTargetFilterSubMenu(FMenuBuilder& InMenuBuilde
 
 ////////////////////////////////////////////////////////////////////////////////////////////////////
 
+TSharedRef<SWidget> STraceStoreWindow::MakeBuildVersionColumnHeaderMenu()
+{
+	FSlateApplication::Get().CloseToolTip();
+
+	FMenuBuilder MenuBuilder(/*bInShouldCloseWindowAfterMenuSelection=*/true, nullptr);
+
+	MenuBuilder.BeginSection("Filter", LOCTEXT("ColumnHeaderMenuSection_BuildVersionFilter", "Build Version Filter"));
+	BuildBuildVersionFilterSubMenu(MenuBuilder);
+	MenuBuilder.EndSection();
+
+	return MenuBuilder.MakeWidget();
+}
+
+////////////////////////////////////////////////////////////////////////////////////////////////////
+
+TSharedRef<SWidget> STraceStoreWindow::MakeBuildVersionFilterMenu()
+{
+	FSlateApplication::Get().CloseToolTip();
+
+	FMenuBuilder MenuBuilder(/*bInShouldCloseWindowAfterMenuSelection=*/true, nullptr);
+
+	MenuBuilder.BeginSection("Filter", LOCTEXT("MenuSection_BuildVersionFilter", "Build Version Filter"));
+	BuildBuildVersionFilterSubMenu(MenuBuilder);
+	MenuBuilder.EndSection();
+
+	return MenuBuilder.MakeWidget();
+}
+
+////////////////////////////////////////////////////////////////////////////////////////////////////
+
+void STraceStoreWindow::BuildBuildVersionFilterSubMenu(FMenuBuilder& InMenuBuilder)
+{
+	FilterByBuildVersion->BuildMenu(InMenuBuilder, *this);
+}
+
+////////////////////////////////////////////////////////////////////////////////////////////////////
+
 TSharedRef<SWidget> STraceStoreWindow::MakeBranchColumnHeaderMenu()
 {
 	FSlateApplication::Get().CloseToolTip();
@@ -2545,14 +2632,14 @@ void STraceStoreWindow::BuildBranchFilterSubMenu(FMenuBuilder& InMenuBuilder)
 
 ////////////////////////////////////////////////////////////////////////////////////////////////////
 
-TSharedRef<SWidget> STraceStoreWindow::MakeVersionColumnHeaderMenu()
+TSharedRef<SWidget> STraceStoreWindow::MakeChangelistColumnHeaderMenu()
 {
 	FSlateApplication::Get().CloseToolTip();
 
 	FMenuBuilder MenuBuilder(/*bInShouldCloseWindowAfterMenuSelection=*/true, nullptr);
 
-	MenuBuilder.BeginSection("Filter", LOCTEXT("ColumnHeaderMenuSection_VersionFilter", "Version Filter"));
-	BuildVersionFilterSubMenu(MenuBuilder);
+	MenuBuilder.BeginSection("Filter", LOCTEXT("ColumnHeaderMenuSection_ChangelistFilter", "Changelist Filter"));
+	BuildChangelistFilterSubMenu(MenuBuilder);
 	MenuBuilder.EndSection();
 
 	return MenuBuilder.MakeWidget();
@@ -2560,14 +2647,14 @@ TSharedRef<SWidget> STraceStoreWindow::MakeVersionColumnHeaderMenu()
 
 ////////////////////////////////////////////////////////////////////////////////////////////////////
 
-TSharedRef<SWidget> STraceStoreWindow::MakeVersionFilterMenu()
+TSharedRef<SWidget> STraceStoreWindow::MakeChangelistFilterMenu()
 {
 	FSlateApplication::Get().CloseToolTip();
 
 	FMenuBuilder MenuBuilder(/*bInShouldCloseWindowAfterMenuSelection=*/true, nullptr);
 
-	MenuBuilder.BeginSection("Filter", LOCTEXT("MenuSection_VersionFilter", "Version Filter"));
-	BuildVersionFilterSubMenu(MenuBuilder);
+	MenuBuilder.BeginSection("Filter", LOCTEXT("MenuSection_ChangelistFilter", "Changelist Filter"));
+	BuildChangelistFilterSubMenu(MenuBuilder);
 	MenuBuilder.EndSection();
 
 	return MenuBuilder.MakeWidget();
@@ -2575,9 +2662,46 @@ TSharedRef<SWidget> STraceStoreWindow::MakeVersionFilterMenu()
 
 ////////////////////////////////////////////////////////////////////////////////////////////////////
 
-void STraceStoreWindow::BuildVersionFilterSubMenu(FMenuBuilder& InMenuBuilder)
+void STraceStoreWindow::BuildChangelistFilterSubMenu(FMenuBuilder& InMenuBuilder)
+{
+	FilterByChangelist->BuildMenu(InMenuBuilder, *this);
+}
+
+////////////////////////////////////////////////////////////////////////////////////////////////////
+
+TSharedRef<SWidget> STraceStoreWindow::MakeEngineVersionColumnHeaderMenu()
 {
-	FilterByVersion->BuildMenu(InMenuBuilder, *this);
+	FSlateApplication::Get().CloseToolTip();
+
+	FMenuBuilder MenuBuilder(/*bInShouldCloseWindowAfterMenuSelection=*/true, nullptr);
+
+	MenuBuilder.BeginSection("Filter", LOCTEXT("ColumnHeaderMenuSection_EngineVersionFilter", "Engine Version Filter"));
+	BuildEngineVersionFilterSubMenu(MenuBuilder);
+	MenuBuilder.EndSection();
+
+	return MenuBuilder.MakeWidget();
+}
+
+////////////////////////////////////////////////////////////////////////////////////////////////////
+
+TSharedRef<SWidget> STraceStoreWindow::MakeEngineVersionFilterMenu()
+{
+	FSlateApplication::Get().CloseToolTip();
+
+	FMenuBuilder MenuBuilder(/*bInShouldCloseWindowAfterMenuSelection=*/true, nullptr);
+
+	MenuBuilder.BeginSection("Filter", LOCTEXT("MenuSection_EngineVersionFilter", "Engine Version Filter"));
+	BuildEngineVersionFilterSubMenu(MenuBuilder);
+	MenuBuilder.EndSection();
+
+	return MenuBuilder.MakeWidget();
+}
+
+////////////////////////////////////////////////////////////////////////////////////////////////////
+
+void STraceStoreWindow::BuildEngineVersionFilterSubMenu(FMenuBuilder& InMenuBuilder)
+{
+	FilterByEngineVersion->BuildMenu(InMenuBuilder, *this);
 }
 
 ////////////////////////////////////////////////////////////////////////////////////////////////////
@@ -2856,11 +2980,17 @@ void STraceStoreWindow::CreateFilters()
 	FilterByBuildTarget = MakeShared<FTraceFilterByBuildTarget>();
 	Filters->Add(FilterByBuildTarget);
 
+	FilterByBuildVersion = MakeShared<FTraceFilterByBuildVersion>();
+	Filters->Add(FilterByBuildVersion);
+
 	FilterByBranch = MakeShared<FTraceFilterByBranch>();
 	Filters->Add(FilterByBranch);
 
-	FilterByVersion = MakeShared<FTraceFilterByVersion>();
-	Filters->Add(FilterByVersion);
+	FilterByChangelist = MakeShared<FTraceFilterByChangelist>();
+	Filters->Add(FilterByChangelist);
+
+	FilterByEngineVersion = MakeShared<FTraceFilterByEngineVersion>();
+	Filters->Add(FilterByEngineVersion);
 
 	FilterBySize = MakeShared<FTraceFilterBySize>();
 	Filters->Add(FilterBySize);
@@ -2875,11 +3005,11 @@ void STraceStoreWindow::HandleItemToStringArray(const FTraceViewModel& InTrace,
 {
 	if (bSearchByCommandLine)
 	{
-		OutSearchStrings.Add(InTrace.CommandLine.ToString());
+		OutSearchStrings.Add(InTrace.GetCommandLine());
 	}
 	else
 	{
-		OutSearchStrings.Add(InTrace.Name.ToString());
+		OutSearchStrings.Add(InTrace.GetName());
 	}
 }
 
@@ -2894,8 +3024,10 @@ void STraceStoreWindow::UpdateFiltering()
 		FilterByAppName->IsEmpty() &&
 		FilterByBuildConfig->IsEmpty() &&
 		FilterByBuildTarget->IsEmpty() &&
+		FilterByBuildVersion->IsEmpty() &&
 		FilterByBranch->IsEmpty() &&
-		FilterByVersion->IsEmpty() &&
+		FilterByChangelist->IsEmpty() &&
+		FilterByEngineVersion->IsEmpty() &&
 		FilterBySize->IsEmpty() &&
 		FilterByStatus->IsEmpty())
 	{
@@ -2926,8 +3058,8 @@ void STraceStoreWindow::UpdateFilterStatsText()
 	uint64 FilterdTotalSize = 0;
 	for (const TSharedPtr<FTraceViewModel>& Trace : FilteredTraceViewModels)
 	{
-		FilterdTotalSize += Trace->Size;
-		if (Trace->bIsLive)
+		FilterdTotalSize += Trace->GetSize();
+		if (Trace->IsLive())
 		{
 			bFilterStatsTextIsDirty = true;
 		}
@@ -2982,57 +3114,74 @@ void STraceStoreWindow::UpdateSorting()
 	{
 		if (SortMode == EColumnSortMode::Ascending)
 		{
-			FilteredTraceViewModels.Sort([](const TSharedPtr<FTraceViewModel>& A, const TSharedPtr<FTraceViewModel>& B)
-				{ return A->Timestamp < B->Timestamp; });
+			FilteredTraceViewModels.Sort(
+				[](const TSharedPtr<FTraceViewModel>& A, const TSharedPtr<FTraceViewModel>& B)
+				{
+					return A->GetTimestamp() < B->GetTimestamp();
+				});
 		}
 		else
 		{
 			FilteredTraceViewModels.Sort([](const TSharedPtr<FTraceViewModel>& A, const TSharedPtr<FTraceViewModel>& B)
-				{ return A->Timestamp > B->Timestamp; });
+				{ return A->GetTimestamp() > B->GetTimestamp(); });
 		}
 	}
 	else if (SortColumn == FTraceListColumns::Name)
 	{
 		if (SortMode == EColumnSortMode::Ascending)
 		{
-			FilteredTraceViewModels.Sort([](const TSharedPtr<FTraceViewModel>& A, const TSharedPtr<FTraceViewModel>& B)
-				{ return A->Name.CompareTo(B->Name) < 0; });
+			FilteredTraceViewModels.Sort(
+				[](const TSharedPtr<FTraceViewModel>& A, const TSharedPtr<FTraceViewModel>& B)
+				{
+					return A->GetNameText().CompareTo(B->GetNameText()) < 0;
+				});
 		}
 		else
 		{
-			FilteredTraceViewModels.Sort([](const TSharedPtr<FTraceViewModel>& A, const TSharedPtr<FTraceViewModel>& B)
-				{ return B->Name.CompareTo(A->Name) < 0; });
+			FilteredTraceViewModels.Sort(
+				[](const TSharedPtr<FTraceViewModel>& A, const TSharedPtr<FTraceViewModel>& B)
+				{
+					return B->GetNameText().CompareTo(A->GetNameText()) < 0;
+				});
 		}
 	}
 	else if (SortColumn == FTraceListColumns::Uri)
 	{
 		if (SortMode == EColumnSortMode::Ascending)
 		{
-			FilteredTraceViewModels.Sort([](const TSharedPtr<FTraceViewModel>& A, const TSharedPtr<FTraceViewModel>& B)
-				{ return A->Uri.CompareTo(B->Uri) < 0; });
+			FilteredTraceViewModels.Sort(
+				[](const TSharedPtr<FTraceViewModel>& A, const TSharedPtr<FTraceViewModel>& B)
+				{
+					return A->GetUriText().CompareTo(B->GetUriText()) < 0;
+				});
 		}
 		else
 		{
-			FilteredTraceViewModels.Sort([](const TSharedPtr<FTraceViewModel>& A, const TSharedPtr<FTraceViewModel>& B)
-				{ return B->Uri.CompareTo(A->Uri) < 0; });
+			FilteredTraceViewModels.Sort(
+				[](const TSharedPtr<FTraceViewModel>& A, const TSharedPtr<FTraceViewModel>& B)
+				{
+					return B->GetUriText().CompareTo(A->GetUriText()) < 0;
+				});
 		}
 	}
 	else if (SortColumn == FTraceListColumns::Platform)
 	{
 		if (SortMode == EColumnSortMode::Ascending)
 		{
-			FilteredTraceViewModels.Sort([](const TSharedPtr<FTraceViewModel>& A, const TSharedPtr<FTraceViewModel>& B)
+			FilteredTraceViewModels.Sort(
+				[](const TSharedPtr<FTraceViewModel>& A, const TSharedPtr<FTraceViewModel>& B)
 				{
-					int32 CompareResult = A->Platform.CompareTo(B->Platform);
-					return CompareResult == 0 ? A->Timestamp < B->Timestamp : CompareResult < 0;
+					int32 CompareResult = A->GetPlatformText().CompareTo(B->GetPlatformText());
+					return CompareResult == 0 ? A->GetTimestamp() < B->GetTimestamp() : CompareResult < 0;
 				});
 		}
 		else
 		{
-			FilteredTraceViewModels.Sort([](const TSharedPtr<FTraceViewModel>& A, const TSharedPtr<FTraceViewModel>& B)
+			FilteredTraceViewModels.Sort(
+				[](const TSharedPtr<FTraceViewModel>& A, const TSharedPtr<FTraceViewModel>& B)
 				{
-					int32 CompareResult = B->Platform.CompareTo(A->Platform);
-					return CompareResult == 0 ? A->Timestamp < B->Timestamp : CompareResult < 0;
+					int32 CompareResult = B->GetPlatformText().CompareTo(A->GetPlatformText());
+					return CompareResult == 0 ? A->GetTimestamp() < B->GetTimestamp() : CompareResult < 0;
 				});
 		}
 	}
@@ -3040,18 +3189,20 @@ void STraceStoreWindow::UpdateSorting()
 	{
 		if (SortMode == EColumnSortMode::Ascending)
 		{
-			FilteredTraceViewModels.Sort([](const TSharedPtr<FTraceViewModel>& A, const TSharedPtr<FTraceViewModel>& B)
+			FilteredTraceViewModels.Sort(
+				[](const TSharedPtr<FTraceViewModel>& A, const TSharedPtr<FTraceViewModel>& B)
 				{
-					int32 CompareResult = A->AppName.CompareTo(B->AppName);
-					return CompareResult == 0 ? A->Timestamp < B->Timestamp : CompareResult < 0;
+					int32 CompareResult = A->GetAppNameText().CompareTo(B->GetAppNameText());
+					return CompareResult == 0 ? A->GetTimestamp() < B->GetTimestamp() : CompareResult < 0;
 				});
 		}
 		else
 		{
-			FilteredTraceViewModels.Sort([](const TSharedPtr<FTraceViewModel>& A, const TSharedPtr<FTraceViewModel>& B)
+			FilteredTraceViewModels.Sort(
+				[](const TSharedPtr<FTraceViewModel>& A, const TSharedPtr<FTraceViewModel>& B)
 				{
-					int32 CompareResult = B->AppName.CompareTo(A->AppName);
-					return CompareResult == 0 ? A->Timestamp < B->Timestamp : CompareResult < 0;
+					int32 CompareResult = B->GetAppNameText().CompareTo(A->GetAppNameText());
+					return CompareResult == 0 ? A->GetTimestamp() < B->GetTimestamp() : CompareResult < 0;
 				});
 		}
 	}
@@ -3059,20 +3210,22 @@ void STraceStoreWindow::UpdateSorting()
 	{
 		if (SortMode == EColumnSortMode::Ascending)
 		{
-			FilteredTraceViewModels.Sort([](const TSharedPtr<FTraceViewModel>& A, const TSharedPtr<FTraceViewModel>& B)
+			FilteredTraceViewModels.Sort(
+				[](const TSharedPtr<FTraceViewModel>& A, const TSharedPtr<FTraceViewModel>& B)
 				{
-					return A->ConfigurationType == B->ConfigurationType ?
-						A->Timestamp < B->Timestamp :
-						A->ConfigurationType < B->ConfigurationType;
+					return A->GetBuildConfiguration() == B->GetBuildConfiguration() ?
+						A->GetTimestamp() < B->GetTimestamp() :
+						A->GetBuildConfiguration() < B->GetBuildConfiguration();
 				});
 		}
 		else
 		{
-			FilteredTraceViewModels.Sort([](const TSharedPtr<FTraceViewModel>& A, const TSharedPtr<FTraceViewModel>& B)
+			FilteredTraceViewModels.Sort(
+				[](const TSharedPtr<FTraceViewModel>& A, const TSharedPtr<FTraceViewModel>& B)
 				{
-					return A->ConfigurationType == B->ConfigurationType ?
-						A->Timestamp < B->Timestamp :
-						A->ConfigurationType > B->ConfigurationType;
+					return A->GetBuildConfiguration() == B->GetBuildConfiguration() ?
+						A->GetTimestamp() < B->GetTimestamp() :
+						A->GetBuildConfiguration() > B->GetBuildConfiguration();
 				});
 		}
 	}
@@ -3080,20 +3233,41 @@ void STraceStoreWindow::UpdateSorting()
 	{
 		if (SortMode == EColumnSortMode::Ascending)
 		{
-			FilteredTraceViewModels.Sort([](const TSharedPtr<FTraceViewModel>& A, const TSharedPtr<FTraceViewModel>& B)
+			FilteredTraceViewModels.Sort(
+				[](const TSharedPtr<FTraceViewModel>& A, const TSharedPtr<FTraceViewModel>& B)
 				{
-					return A->TargetType == B->TargetType ?
-						A->Timestamp < B->Timestamp :
-						A->TargetType < B->TargetType;
+					return A->GetBuildTarget() == B->GetBuildTarget() ?
+						A->GetTimestamp() < B->GetTimestamp() :
+						A->GetBuildTarget() < B->GetBuildTarget();
 				});
 		}
 		else
 		{
-			FilteredTraceViewModels.Sort([](const TSharedPtr<FTraceViewModel>& A, const TSharedPtr<FTraceViewModel>& B)
+			FilteredTraceViewModels.Sort(
+				[](const TSharedPtr<FTraceViewModel>& A, const TSharedPtr<FTraceViewModel>& B)
 				{
-					return A->TargetType == B->TargetType ?
-						A->Timestamp < B->Timestamp :
-						A->TargetType > B->TargetType;
+					return A->GetBuildTarget() == B->GetBuildTarget() ?
+						A->GetTimestamp() < B->GetTimestamp() :
+						A->GetBuildTarget() > B->GetBuildTarget();
+				});
+		}
+	}
+	else if (SortColumn == FTraceListColumns::BuildVersion)
+	{
+		if (SortMode == EColumnSortMode::Ascending)
+		{
+			FilteredTraceViewModels.Sort(
+				[](const TSharedPtr<FTraceViewModel>& A, const TSharedPtr<FTraceViewModel>& B)
+				{
+					return A->GetBuildVersionText().CompareTo(B->GetBuildVersionText()) < 0;
+				});
+		}
+		else
+		{
+			FilteredTraceViewModels.Sort(
+				[](const TSharedPtr<FTraceViewModel>& A, const TSharedPtr<FTraceViewModel>& B)
+				{
+					return B->GetBuildVersionText().CompareTo(A->GetBuildVersionText()) < 0;
 				});
 		}
 	}
@@ -3101,65 +3275,108 @@ void STraceStoreWindow::UpdateSorting()
 	{
 		if (SortMode == EColumnSortMode::Ascending)
 		{
-			FilteredTraceViewModels.Sort([](const TSharedPtr<FTraceViewModel>& A, const TSharedPtr<FTraceViewModel>& B)
-				{ return A->Branch.CompareTo(B->Branch) < 0; });
+			FilteredTraceViewModels.Sort(
+				[](const TSharedPtr<FTraceViewModel>& A, const TSharedPtr<FTraceViewModel>& B)
+				{
+					return A->GetBuildBranchText().CompareTo(B->GetBuildBranchText()) < 0;
+				});
 		}
 		else
 		{
-			FilteredTraceViewModels.Sort([](const TSharedPtr<FTraceViewModel>& A, const TSharedPtr<FTraceViewModel>& B)
-				{ return B->Branch.CompareTo(A->Branch) < 0; });
+			FilteredTraceViewModels.Sort(
+				[](const TSharedPtr<FTraceViewModel>& A, const TSharedPtr<FTraceViewModel>& B)
+				{
+					return B->GetBuildBranchText().CompareTo(A->GetBuildBranchText()) < 0;
+				});
 		}
 	}
-	else if (SortColumn == FTraceListColumns::BuildVersion)
+	else if (SortColumn == FTraceListColumns::Changelist)
 	{
 		if (SortMode == EColumnSortMode::Ascending)
 		{
-			FilteredTraceViewModels.Sort([](const TSharedPtr<FTraceViewModel>& A, const TSharedPtr<FTraceViewModel>& B)
-				{ return A->BuildVersion.CompareTo(B->BuildVersion) < 0; });
+			FilteredTraceViewModels.Sort(
+				[](const TSharedPtr<FTraceViewModel>& A, const TSharedPtr<FTraceViewModel>& B)
+				{
+					return A->GetChangelist() < B->GetChangelist();
+				});
 		}
 		else
 		{
-			FilteredTraceViewModels.Sort([](const TSharedPtr<FTraceViewModel>& A, const TSharedPtr<FTraceViewModel>& B)
-				{ return B->BuildVersion.CompareTo(A->BuildVersion) < 0; });
+			FilteredTraceViewModels.Sort(
+				[](const TSharedPtr<FTraceViewModel>& A, const TSharedPtr<FTraceViewModel>& B)
+				{
+					return B->GetChangelist() < A->GetChangelist();
+				});
+		}
+	}
+	else if (SortColumn == FTraceListColumns::EngineVersion)
+	{
+		if (SortMode == EColumnSortMode::Ascending)
+		{
+			FilteredTraceViewModels.Sort(
+				[](const TSharedPtr<FTraceViewModel>& A, const TSharedPtr<FTraceViewModel>& B)
+				{
+					return A->GetEngineVersionNumber() < B->GetEngineVersionNumber();
+				});
+		}
+		else
+		{
+			FilteredTraceViewModels.Sort(
+				[](const TSharedPtr<FTraceViewModel>& A, const TSharedPtr<FTraceViewModel>& B)
+				{
+					return B->GetEngineVersionNumber() < A->GetEngineVersionNumber();
+				});
 		}
 	}
 	else if (SortColumn == FTraceListColumns::Size)
 	{
 		if (SortMode == EColumnSortMode::Ascending)
 		{
-			FilteredTraceViewModels.Sort([](const TSharedPtr<FTraceViewModel>& A, const TSharedPtr<FTraceViewModel>& B)
-				{ return A->Size < B->Size; });
+			FilteredTraceViewModels.Sort(
+				[](const TSharedPtr<FTraceViewModel>& A, const TSharedPtr<FTraceViewModel>& B)
+				{
+					return A->GetSize() < B->GetSize();
+				});
 		}
 		else
 		{
-			FilteredTraceViewModels.Sort([](const TSharedPtr<FTraceViewModel>& A, const TSharedPtr<FTraceViewModel>& B)
-				{ return A->Size > B->Size; });
+			FilteredTraceViewModels.Sort(
+				[](const TSharedPtr<FTraceViewModel>& A, const TSharedPtr<FTraceViewModel>& B)
+				{
+					return A->GetSize() > B->GetSize();
+				});
 		}
 	}
 	else if (SortColumn == FTraceListColumns::Status)
 	{
 		if (SortMode == EColumnSortMode::Ascending)
 		{
-			FilteredTraceViewModels.Sort([](const TSharedPtr<FTraceViewModel>& A, const TSharedPtr<FTraceViewModel>& B)
+			FilteredTraceViewModels.Sort(
+				[](const TSharedPtr<FTraceViewModel>& A, const TSharedPtr<FTraceViewModel>& B)
 				{
-					return A->bIsLive == B->bIsLive ?
-						A->Timestamp < B->Timestamp :
-						B->bIsLive;
+					return A->IsLive() == B->IsLive() ?
+						A->GetTimestamp() < B->GetTimestamp() :
+						B->IsLive();
 				});
 		}
 		else
 		{
-			FilteredTraceViewModels.Sort([](const TSharedPtr<FTraceViewModel>& A, const TSharedPtr<FTraceViewModel>& B)
+			FilteredTraceViewModels.Sort(
+				[](const TSharedPtr<FTraceViewModel>& A, const TSharedPtr<FTraceViewModel>& B)
 				{
-					return A->bIsLive == B->bIsLive ?
-						A->Timestamp < B->Timestamp :
-						A->bIsLive;
+					return A->IsLive() == B->IsLive() ?
+						A->GetTimestamp() < B->GetTimestamp() :
+						A->IsLive();
 				});
 		}
 	}
 	else
 	{
-		Algo::SortBy(FilteredTraceViewModels, &FTraceViewModel::Timestamp);
+		Algo::SortBy(FilteredTraceViewModels,
+			[](const TSharedPtr<FTraceViewModel>& Trace)
+			{
+				return Trace->GetTimestamp();
+			});
 	}
 }
 
diff --git a/Engine/Source/Developer/TraceInsightsFrontend/Private/InsightsFrontend/Widgets/STraceStoreWindow.h b/Engine/Source/Developer/TraceInsightsFrontend/Private/InsightsFrontend/Widgets/STraceStoreWindow.h
--- a/Engine/Source/Developer/TraceInsightsFrontend/Private/InsightsFrontend/Widgets/STraceStoreWindow.h
+++ b/Engine/Source/Developer/TraceInsightsFrontend/Private/InsightsFrontend/Widgets/STraceStoreWindow.h
@@ -43,15 +43,17 @@ struct FStoreBrowserTraceInfo;
 
 class FTableImporter;
 
-struct FTraceViewModel;
+class FTraceViewModel;
 struct FTraceDirectoryModel;
 
 class FTraceFilterByPlatform;
 class FTraceFilterByAppName;
 class FTraceFilterByBuildConfig;
 class FTraceFilterByBuildTarget;
+class FTraceFilterByBuildVersion;
 class FTraceFilterByBranch;
-class FTraceFilterByVersion;
+class FTraceFilterByChangelist;
+class FTraceFilterByEngineVersion;
 class FTraceFilterBySize;
 class FTraceFilterByStatus;
 
@@ -202,13 +204,21 @@ private:
 	TSharedRef<SWidget> MakeBuildTargetFilterMenu();
 	void BuildBuildTargetFilterSubMenu(FMenuBuilder& InMenuBuilder);
 
+	TSharedRef<SWidget> MakeBuildVersionColumnHeaderMenu();
+	TSharedRef<SWidget> MakeBuildVersionFilterMenu();
+	void BuildBuildVersionFilterSubMenu(FMenuBuilder& InMenuBuilder);
+
 	TSharedRef<SWidget> MakeBranchColumnHeaderMenu();
 	TSharedRef<SWidget> MakeBranchFilterMenu();
 	void BuildBranchFilterSubMenu(FMenuBuilder& InMenuBuilder);
 
-	TSharedRef<SWidget> MakeVersionColumnHeaderMenu();
-	TSharedRef<SWidget> MakeVersionFilterMenu();
-	void BuildVersionFilterSubMenu(FMenuBuilder& InMenuBuilder);
+	TSharedRef<SWidget> MakeChangelistColumnHeaderMenu();
+	TSharedRef<SWidget> MakeChangelistFilterMenu();
+	void BuildChangelistFilterSubMenu(FMenuBuilder& InMenuBuilder);
+
+	TSharedRef<SWidget> MakeEngineVersionColumnHeaderMenu();
+	TSharedRef<SWidget> MakeEngineVersionFilterMenu();
+	void BuildEngineVersionFilterSubMenu(FMenuBuilder& InMenuBuilder);
 
 	TSharedRef<SWidget> MakeSizeColumnHeaderMenu();
 	TSharedRef<SWidget> MakeSizeFilterMenu();
@@ -426,8 +436,10 @@ private:
 	TSharedPtr<FTraceFilterByAppName> FilterByAppName;
 	TSharedPtr<FTraceFilterByBuildConfig> FilterByBuildConfig;
 	TSharedPtr<FTraceFilterByBuildTarget> FilterByBuildTarget;
+	TSharedPtr<FTraceFilterByBuildVersion> FilterByBuildVersion;
 	TSharedPtr<FTraceFilterByBranch> FilterByBranch;
-	TSharedPtr<FTraceFilterByVersion> FilterByVersion;
+	TSharedPtr<FTraceFilterByChangelist> FilterByChangelist;
+	TSharedPtr<FTraceFilterByEngineVersion> FilterByEngineVersion;
 	TSharedPtr<FTraceFilterBySize> FilterBySize;
 	TSharedPtr<FTraceFilterByStatus> FilterByStatus;
 
diff --git a/Engine/Source/Developer/TraceServices/Private/Analyzers/DiagnosticsAnalysis.cpp b/Engine/Source/Developer/TraceServices/Private/Analyzers/DiagnosticsAnalysis.cpp
--- a/Engine/Source/Developer/TraceServices/Private/Analyzers/DiagnosticsAnalysis.cpp
+++ b/Engine/Source/Developer/TraceServices/Private/Analyzers/DiagnosticsAnalysis.cpp
@@ -83,9 +83,10 @@ bool FDiagnosticsAnalyzer::OnEvent(uint16 RouteId, EStyle Style, const FOnEventC
 		EventData.GetString("Platform", SessionInfo.Platform);
 		EventData.GetString("AppName", SessionInfo.AppName);
 		EventData.GetString("ProjectName", SessionInfo.ProjectName);
-		EventData.GetString("CommandLine", SessionInfo.CommandLine);
-		EventData.GetString("Branch", SessionInfo.Branch);
+		EventData.GetString("EngineVersion", SessionInfo.EngineVersion);
 		EventData.GetString("BuildVersion", SessionInfo.BuildVersion);
+		EventData.GetString("Branch", SessionInfo.Branch);
+		EventData.GetString("CommandLine", SessionInfo.CommandLine);
 		EventData.GetString("VFSPaths", SessionInfo.VFSPaths);
 		SessionInfo.Changelist = EventData.GetValue<uint32>("Changelist", 0);
 		SessionInfo.ConfigurationType = (EBuildConfiguration)EventData.GetValue<uint8>("ConfigurationType");
diff --git a/Engine/Source/Developer/TraceServices/Public/TraceServices/Model/Diagnostics.h b/Engine/Source/Developer/TraceServices/Public/TraceServices/Model/Diagnostics.h
--- a/Engine/Source/Developer/TraceServices/Public/TraceServices/Model/Diagnostics.h
+++ b/Engine/Source/Developer/TraceServices/Public/TraceServices/Model/Diagnostics.h
@@ -16,9 +16,10 @@ struct FSessionInfo
 	FString Platform;
 	FString AppName;
 	FString ProjectName;
-	FString CommandLine;
-	FString Branch;
+	FString EngineVersion;
 	FString BuildVersion;
+	FString Branch;
+	FString CommandLine;
 	FString VFSPaths;
 	uint32 Changelist = 0;
 	EBuildConfiguration ConfigurationType = EBuildConfiguration::Unknown;
