Fixed static analysis warning.

#rnx

diff --git a/Engine/Source/Developer/TraceInsightsFrontend/Private/InsightsFrontend/ViewModels/TraceViewModel.cpp b/Engine/Source/Developer/TraceInsightsFrontend/Private/InsightsFrontend/ViewModels/TraceViewModel.cpp
--- a/Engine/Source/Developer/TraceInsightsFrontend/Private/InsightsFrontend/ViewModels/TraceViewModel.cpp
+++ b/Engine/Source/Developer/TraceInsightsFrontend/Private/InsightsFrontend/ViewModels/TraceViewModel.cpp
@@ -79,7 +79,7 @@ void FTraceViewModel::SetEngineVersion(const FString& InEngineVersion)
 			{
 				InOutVersion.Reset();
 			}
-			return uint64(VersionNumber + 1); // +1 is to sort "5.0" after "5"
+			return uint64(FMath::Max(0, VersionNumber)) + 1; // +1 is to sort "5.0" after "5"
 		};
 
 	uint64 Major = ReadVersionNumber(Version);
