[Insights] Frontend: Fixed EngineVersion and VFSPaths metadata not being updated in trace session's tooltips.

#rnx

diff --git a/Engine/Source/Developer/TraceInsightsFrontend/Private/InsightsFrontend/StoreService/StoreBrowser.cpp b/Engine/Source/Developer/TraceInsightsFrontend/Private/InsightsFrontend/StoreService/StoreBrowser.cpp
--- a/Engine/Source/Developer/TraceInsightsFrontend/Private/InsightsFrontend/StoreService/StoreBrowser.cpp
+++ b/Engine/Source/Developer/TraceInsightsFrontend/Private/InsightsFrontend/StoreService/StoreBrowser.cpp
@@ -510,7 +510,7 @@ void FStoreBrowser::UpdateTraces()
 
 	StopwatchTotal.Stop();
 	const uint64 TotalTime = StopwatchTotal.GetAccumulatedTimeMs();
-	if (TotalTime > 5)
+	if (TotalTime > 100)
 	{
 		UE_LOG(LogInsightsFrontend, Log, TEXT("[StoreBrowser] Updated in %llu ms (%llu + %llu + %llu + %llu)."),
 			TotalTime, Step1, Step2 - Step1, Step3 - Step2, TotalTime - Step3);
@@ -636,12 +636,14 @@ void FStoreBrowser::UpdateMetadata(TSharedPtr<FStoreBrowserTraceInfo> TraceInfoP
 			Trace.Platform = Analyzer.Platform;
 			Trace.AppName = Analyzer.AppName;
 			Trace.ProjectName = Analyzer.ProjectName;
-			Trace.CommandLine = Analyzer.CommandLine;
-			Trace.Branch = Analyzer.Branch;
-			Trace.BuildVersion = Analyzer.BuildVersion;
-			Trace.Changelist = Analyzer.Changelist;
 			Trace.ConfigurationType = static_cast<EBuildConfiguration>(Analyzer.ConfigurationType);
 			Trace.TargetType = static_cast<EBuildTargetType>(Analyzer.TargetType);
+			Trace.BuildVersion = Analyzer.BuildVersion;
+			Trace.Branch = Analyzer.Branch;
+			Trace.Changelist = Analyzer.Changelist;
+			Trace.EngineVersion = Analyzer.EngineVersion;
+			Trace.CommandLine = Analyzer.CommandLine;
+			Trace.VFSPaths = Analyzer.VFSPaths;
 		}
 
 		if (DataStream.Status == FDataStream::EReadStatus::StoppedByTimeLimit)
